#summary Fixed some ensure cases with names deviating between GetNodeTitle and GetNativeNodeTitle.

#change K2Node_LiveEditObject will no longer localize the number.
#change Fixed a FText::Format argument that was not set correctly in EnvQueryTest
#change Fixed an incorrect FText::Format argument in AnimGraphNode_SequencePlayer.

#ttp 332386 - EDITOR: MAIN: ENSURE: Right clicking in blueprint event graph triggers ensure

diff --git a/Engine/Plugins/Experimental/LiveEditor/Source/LiveEditor/Private/K2Node_LiveEditObject.cpp b/Engine/Plugins/Experimental/LiveEditor/Source/LiveEditor/Private/K2Node_LiveEditObject.cpp
--- a/Engine/Plugins/Experimental/LiveEditor/Source/LiveEditor/Private/K2Node_LiveEditObject.cpp
+++ b/Engine/Plugins/Experimental/LiveEditor/Source/LiveEditor/Private/K2Node_LiveEditObject.cpp
@@ -187,9 +187,11 @@ FText UK2Node_LiveEditObject::GetNodeTitle(ENodeTitleType::Type TitleType) const
 		SpawnString = FText::FromString(BaseClassPin->DefaultObject->GetName());
 	}
 
+	FNumberFormattingOptions NumberOptions;
+	NumberOptions.UseGrouping = false;
 	FFormatNamedArguments Args;
 	Args.Add(TEXT("SpawnString"), SpawnString);
-	Args.Add(TEXT("ID"), GetUniqueID());
+	Args.Add(TEXT("ID"), FText::AsNumber(GetUniqueID(), &NumberOptions));
 	return FText::Format(NSLOCTEXT("K2Node", "LiveEditObject", "LiveEditObject {SpawnString}_{ID}"), Args );
 }
 
@@ -798,4 +800,4 @@ UClass* UK2Node_LiveEditObject::GetClassToSpawn(const TArray<UEdGraphPin*>* InPi
 	return UseSpawnClass;
 }
 
-#undef LOCTEXT_NAMESPACE
\ No newline at end of file
+#undef LOCTEXT_NAMESPACE
diff --git a/Engine/Source/Editor/AnimGraph/Private/AnimGraphNode_SequencePlayer.cpp b/Engine/Source/Editor/AnimGraph/Private/AnimGraphNode_SequencePlayer.cpp
--- a/Engine/Source/Editor/AnimGraph/Private/AnimGraphNode_SequencePlayer.cpp
+++ b/Engine/Source/Editor/AnimGraph/Private/AnimGraphNode_SequencePlayer.cpp
@@ -83,8 +83,9 @@ FText UAnimGraphNode_SequencePlayer::GetNodeTitle(ENodeTitleType::Type TitleType
 	{
 		FFormatNamedArguments Args;
 		Args.Add(TEXT("Title"), Title);
+		Args.Add(TEXT("SyncGroup"), FText::FromName(SyncGroup.GroupName));
 
-		Title = FText::Format(LOCTEXT("SequenceNodeGroupWithSubtitle", "{Title}\nSync group {Title}"), Args);
+		Title = FText::Format(LOCTEXT("SequenceNodeGroupWithSubtitle", "{Title}\nSync group {SyncGroup}"), Args);
 	}
 
 	return Title;
diff --git a/Engine/Source/Runtime/Engine/Private/AI/EnvironmentQuery/EnvQueryTest.cpp b/Engine/Source/Runtime/Engine/Private/AI/EnvironmentQuery/EnvQueryTest.cpp
--- a/Engine/Source/Runtime/Engine/Private/AI/EnvironmentQuery/EnvQueryTest.cpp
+++ b/Engine/Source/Runtime/Engine/Private/AI/EnvironmentQuery/EnvQueryTest.cpp
@@ -156,7 +156,7 @@ FText UEnvQueryTest::DescribeFloatTestParams() const
 		FText InvalidConditionDesc = LOCTEXT("InvalidCondition", "invalid condition");
 		ParamDesc =
 			(Condition == EEnvTestCondition::AtLeast) ? FText::Format(LOCTEXT("AtLeastFiltered", "at least {Filter}"), Args) :
-			(Condition == EEnvTestCondition::UpTo) ? FText::Format(LOCTEXT("UpToFiltered", "up to %s"), Args) :
+			(Condition == EEnvTestCondition::UpTo) ? FText::Format(LOCTEXT("UpToFiltered", "up to {Filter}"), Args) :
 			InvalidConditionDesc;
 	}
 
diff --git a/Engine/Source/Runtime/Engine/Private/EdGraph/EdGraph.cpp b/Engine/Source/Runtime/Engine/Private/EdGraph/EdGraph.cpp
--- a/Engine/Source/Runtime/Engine/Private/EdGraph/EdGraph.cpp
+++ b/Engine/Source/Runtime/Engine/Private/EdGraph/EdGraph.cpp
@@ -635,11 +635,11 @@ FText UEdGraphNode::GetNodeSearchTitle() const
 
 	FText LocalizedTitle = GetNodeTitle(ENodeTitleType::ListView);
 	FText NativeTitle = FText::FromString(GetNodeNativeTitle(ENodeTitleType::ListView));
-//	@TODO: vet the performance implications of this
-// 	if(FInternationalization::Get().GetCurrentCulture()->GetName() == TEXT("en_US"))
-// 	{
-// 		ensure( LocalizedTitle.CompareTo(NativeTitle) == 0 );
-// 	}
+
+	if(FInternationalization::Get().GetCurrentCulture()->GetName() == TEXT("en_US"))
+	{
+		ensure( LocalizedTitle.CompareTo(NativeTitle) == 0 );
+	}
 
 	FFormatNamedArguments Args;
 	Args.Add(TEXT("LocalizedTitle"), LocalizedTitle);
