Add support for LiveEditing Arrays

diff --git a/Engine/Plugins/Experimental/LiveEditor/Source/LiveEditor/Private/K2Node_LiveEditObject.cpp b/Engine/Plugins/Experimental/LiveEditor/Source/LiveEditor/Private/K2Node_LiveEditObject.cpp
--- a/Engine/Plugins/Experimental/LiveEditor/Source/LiveEditor/Private/K2Node_LiveEditObject.cpp
+++ b/Engine/Plugins/Experimental/LiveEditor/Source/LiveEditor/Private/K2Node_LiveEditObject.cpp
@@ -38,6 +38,11 @@ namespace UK2Node_LiveEditObjectStatics
 			RemainingTokens = FString(TEXT(""));
 		}
 
+		if ( FirstToken.FindChar( '[', SplitIndex ) )
+		{
+			FirstToken = FirstToken.LeftChop( FirstToken.Len()-SplitIndex );
+		}
+
 		for (TFieldIterator<UProperty> PropertyIt(InStruct, EFieldIteratorFlags::IncludeSuper); PropertyIt; ++PropertyIt)
 		{
 			UProperty* Property = *PropertyIt;
diff --git a/Engine/Plugins/Experimental/LiveEditor/Source/LiveEditor/Private/LiveEditorKismetLibrary.cpp b/Engine/Plugins/Experimental/LiveEditor/Source/LiveEditor/Private/LiveEditorKismetLibrary.cpp
--- a/Engine/Plugins/Experimental/LiveEditor/Source/LiveEditor/Private/LiveEditorKismetLibrary.cpp
+++ b/Engine/Plugins/Experimental/LiveEditor/Source/LiveEditor/Private/LiveEditorKismetLibrary.cpp
@@ -1,6 +1,7 @@
 
 #include "LiveEditorPrivatePCH.h"
 #include "LiveEditorKismetLibrary.h"
+#include "DefaultValueHelper.h"
 
 namespace nLiveEditorKismetLibrary
 {
@@ -49,7 +50,7 @@ namespace nLiveEditorKismetLibrary
 		}
 	}
 
-	static UProperty *GetPropertyByNameRecurse( UStruct *InStruct, const FString &TokenString, void ** hContainerPtr )
+	static UProperty *GetPropertyByNameRecurse( UStruct *InStruct, const FString &TokenString, void ** hContainerPtr, int32 &OutArrayIndex )
 	{
 		FString FirstToken;
 		FString RemainingTokens;
@@ -65,6 +66,17 @@ namespace nLiveEditorKismetLibrary
 			RemainingTokens = FString(TEXT(""));
 		}
 
+		//get the array index if there is any
+		int32 ArrayIndex = 0;
+		if ( FirstToken.FindChar( '[', SplitIndex ) )
+		{
+			FString ArrayIndexString = FirstToken.RightChop( SplitIndex+1 );
+			ArrayIndexString = ArrayIndexString.LeftChop( 1 );
+			FDefaultValueHelper::ParseInt( ArrayIndexString, ArrayIndex );
+
+			FirstToken = FirstToken.LeftChop( FirstToken.Len()-SplitIndex );
+		}
+
 		for (TFieldIterator<UProperty> PropertyIt(InStruct, EFieldIteratorFlags::IncludeSuper); PropertyIt; ++PropertyIt)
 		{
 			UProperty* Property = *PropertyIt;
@@ -74,6 +86,7 @@ namespace nLiveEditorKismetLibrary
 				if ( RemainingTokens.Len() == 0 )
 				{
 					check( *hContainerPtr != NULL );
+					OutArrayIndex = ArrayIndex;
 					return Property;
 				}
 				else
@@ -82,8 +95,8 @@ namespace nLiveEditorKismetLibrary
 					if ( StructProp )
 					{
 						check( *hContainerPtr != NULL );
-						*hContainerPtr = Property->ContainerPtrToValuePtr<void>(*hContainerPtr);
-						return GetPropertyByNameRecurse( StructProp->Struct, RemainingTokens, hContainerPtr );
+						*hContainerPtr = Property->ContainerPtrToValuePtr<void>( *hContainerPtr, ArrayIndex );
+						return GetPropertyByNameRecurse( StructProp->Struct, RemainingTokens, hContainerPtr, OutArrayIndex );
 					}
 				}
 			}
@@ -97,7 +110,8 @@ namespace nLiveEditorKismetLibrary
 			return;
 
 		void *ContainerPtr = Target;
-		UProperty *Prop = GetPropertyByNameRecurse( Target->GetClass(), PropertyName, &ContainerPtr );
+		int32 ArrayIndex = 0;
+		UProperty *Prop = GetPropertyByNameRecurse( Target->GetClass(), PropertyName, &ContainerPtr, ArrayIndex );
 		if ( Prop == NULL )
 		{
 			AActor *AsActor = Cast<AActor>(Target);
@@ -124,7 +138,7 @@ namespace nLiveEditorKismetLibrary
 						*/
 
 						ContainerPtr = Component;
-						Prop = GetPropertyByNameRecurse( Component->GetClass(), ComponentPropertyName, &ContainerPtr );
+						Prop = GetPropertyByNameRecurse( Component->GetClass(), ComponentPropertyName, &ContainerPtr, ArrayIndex );
 					}
 				}
 			}
@@ -150,72 +164,72 @@ namespace nLiveEditorKismetLibrary
 		if ( Prop->IsA( UByteProperty::StaticClass() ) )
 		{
 			UByteProperty *NumericProp = CastChecked<UByteProperty>(Prop);
-			uint8 CurValue = NumericProp->GetPropertyValue_InContainer(ContainerPtr);
+			uint8 CurValue = NumericProp->GetPropertyValue_InContainer(ContainerPtr, ArrayIndex);
 			uint8 NewValue = CalculateNewValue<uint8>( CurValue, ControlType, Delta, MidiValue, bShouldClamp, ClampMin, ClampMax );
-			NumericProp->SetPropertyValue_InContainer(ContainerPtr, NewValue);
+			NumericProp->SetPropertyValue_InContainer(ContainerPtr, NewValue, ArrayIndex);
 		}
 		else if ( Prop->IsA( UInt8Property::StaticClass() ) )
 		{
 			UInt8Property *NumericProp = CastChecked<UInt8Property>(Prop);
-			uint8 CurValue = NumericProp->GetPropertyValue_InContainer(ContainerPtr);
+			uint8 CurValue = NumericProp->GetPropertyValue_InContainer(ContainerPtr, ArrayIndex);
 			uint8 NewValue = CalculateNewValue<uint8>( CurValue, ControlType, Delta, MidiValue, bShouldClamp, ClampMin, ClampMax );
-			NumericProp->SetPropertyValue_InContainer(ContainerPtr, NewValue);
+			NumericProp->SetPropertyValue_InContainer(ContainerPtr, NewValue, ArrayIndex);
 		}
 		else if ( Prop->IsA( UInt16Property::StaticClass() ) )
 		{
 			UInt16Property *NumericProp = CastChecked<UInt16Property>(Prop);
-			int16 CurValue = NumericProp->GetPropertyValue_InContainer(ContainerPtr);
+			int16 CurValue = NumericProp->GetPropertyValue_InContainer(ContainerPtr, ArrayIndex);
 			int16 NewValue = CalculateNewValue<int16>( CurValue, ControlType, Delta, MidiValue, bShouldClamp, ClampMin, ClampMax );
-			NumericProp->SetPropertyValue_InContainer(ContainerPtr, NewValue);
+			NumericProp->SetPropertyValue_InContainer(ContainerPtr, NewValue, ArrayIndex);
 		}
 		else if ( Prop->IsA( UIntProperty::StaticClass() ) )
 		{
 			UIntProperty *NumericProp = CastChecked<UIntProperty>(Prop);
-			int32 CurValue = NumericProp->GetPropertyValue_InContainer(ContainerPtr);
+			int32 CurValue = NumericProp->GetPropertyValue_InContainer(ContainerPtr, ArrayIndex);
 			int32 NewValue = CalculateNewValue<int32>( CurValue, ControlType, Delta, MidiValue, bShouldClamp, ClampMin, ClampMax );
-			NumericProp->SetPropertyValue_InContainer(ContainerPtr, NewValue);
+			NumericProp->SetPropertyValue_InContainer(ContainerPtr, NewValue, ArrayIndex);
 		}
 		else if ( Prop->IsA( UInt64Property::StaticClass() ) )
 		{
 			UInt64Property *NumericProp = CastChecked<UInt64Property>(Prop);
-			int64 CurValue = NumericProp->GetPropertyValue_InContainer(ContainerPtr);
+			int64 CurValue = NumericProp->GetPropertyValue_InContainer(ContainerPtr, ArrayIndex);
 			int64 NewValue = CalculateNewValue<int64>( CurValue, ControlType, Delta, MidiValue, bShouldClamp, ClampMin, ClampMax );
-			NumericProp->SetPropertyValue_InContainer(ContainerPtr, NewValue);
+			NumericProp->SetPropertyValue_InContainer(ContainerPtr, NewValue, ArrayIndex);
 		}
 		else if ( Prop->IsA( UUInt16Property::StaticClass() ) )
 		{
 			UUInt16Property *NumericProp = CastChecked<UUInt16Property>(Prop);
-			uint16 CurValue = NumericProp->GetPropertyValue_InContainer(ContainerPtr);
+			uint16 CurValue = NumericProp->GetPropertyValue_InContainer(ContainerPtr, ArrayIndex);
 			uint16 NewValue = CalculateNewValue<uint16>( CurValue, ControlType, Delta, MidiValue, bShouldClamp, ClampMin, ClampMax );
-			NumericProp->SetPropertyValue_InContainer(ContainerPtr, NewValue);
+			NumericProp->SetPropertyValue_InContainer(ContainerPtr, NewValue, ArrayIndex);
 		}
 		else if ( Prop->IsA( UUInt32Property::StaticClass() ) )
 		{
 			UUInt32Property *NumericProp = CastChecked<UUInt32Property>(Prop);
-			uint32 CurValue = NumericProp->GetPropertyValue_InContainer(ContainerPtr);
+			uint32 CurValue = NumericProp->GetPropertyValue_InContainer(ContainerPtr, ArrayIndex);
 			uint32 NewValue = CalculateNewValue<uint32>( CurValue, ControlType, Delta, MidiValue, bShouldClamp, ClampMin, ClampMax );
-			NumericProp->SetPropertyValue_InContainer(ContainerPtr, NewValue);
+			NumericProp->SetPropertyValue_InContainer(ContainerPtr, NewValue, ArrayIndex);
 		}
 		else if ( Prop->IsA( UInt64Property::StaticClass() ) )
 		{
 			UInt64Property *NumericProp = CastChecked<UInt64Property>(Prop);
-			uint64 CurValue = NumericProp->GetPropertyValue_InContainer(ContainerPtr);
+			uint64 CurValue = NumericProp->GetPropertyValue_InContainer(ContainerPtr, ArrayIndex);
 			uint64 NewValue = CalculateNewValue<uint64>( CurValue, ControlType, Delta, MidiValue, bShouldClamp, ClampMin, ClampMax );
-			NumericProp->SetPropertyValue_InContainer(ContainerPtr, NewValue);
+			NumericProp->SetPropertyValue_InContainer(ContainerPtr, NewValue, ArrayIndex);
 		}
 		else if ( Prop->IsA( UFloatProperty::StaticClass() ) )
 		{
 			UFloatProperty *NumericProp = CastChecked<UFloatProperty>(Prop);
-			float CurValue = NumericProp->GetPropertyValue_InContainer(ContainerPtr);
+			float CurValue = NumericProp->GetPropertyValue_InContainer(ContainerPtr, ArrayIndex);
 			float NewValue = CalculateNewValue<float>( CurValue, ControlType, Delta, MidiValue, bShouldClamp, ClampMin, ClampMax );
-			NumericProp->SetPropertyValue_InContainer(ContainerPtr, NewValue);
+			NumericProp->SetPropertyValue_InContainer(ContainerPtr, NewValue, ArrayIndex);
 		}
 		else if ( Prop->IsA( UDoubleProperty::StaticClass() ) )
 		{
 			UDoubleProperty *NumericProp = CastChecked<UDoubleProperty>(Prop);
-			double CurValue = NumericProp->GetPropertyValue_InContainer(ContainerPtr);
+			double CurValue = NumericProp->GetPropertyValue_InContainer(ContainerPtr, ArrayIndex);
 			double NewValue = CalculateNewValue<double>( CurValue, ControlType, Delta, MidiValue, bShouldClamp, ClampMin, ClampMax );
-			NumericProp->SetPropertyValue_InContainer(ContainerPtr, NewValue);
+			NumericProp->SetPropertyValue_InContainer(ContainerPtr, NewValue, ArrayIndex);
 		}
 	}
 
@@ -225,7 +239,8 @@ namespace nLiveEditorKismetLibrary
 			return;
 
 		void *ArchetypeContainerPtr = Archetype;
-		UProperty *Prop = GetPropertyByNameRecurse( Target->GetClass(), PropertyName.ToString(), &ArchetypeContainerPtr );
+		int32 ArrayIndex;
+		UProperty *Prop = GetPropertyByNameRecurse( Target->GetClass(), PropertyName.ToString(), &ArchetypeContainerPtr, ArrayIndex );
 		if ( !Prop->IsA( UNumericProperty::StaticClass() ) )
 			return;
 		check(ArchetypeContainerPtr != NULL);
@@ -234,62 +249,62 @@ namespace nLiveEditorKismetLibrary
 		if ( Prop->IsA( UByteProperty::StaticClass() ) )
 		{
 			UByteProperty *NumericProp = CastChecked<UByteProperty>(Prop);
-			uint8 ArchetypeVal = NumericProp->GetPropertyValue_InContainer(ArchetypeContainerPtr);
-			NumericProp->SetPropertyValue_InContainer(TargetContainerPtr, ArchetypeVal);
+			uint8 ArchetypeVal = NumericProp->GetPropertyValue_InContainer(ArchetypeContainerPtr, ArrayIndex);
+			NumericProp->SetPropertyValue_InContainer(TargetContainerPtr, ArchetypeVal, ArrayIndex);
 		}
 		else if ( Prop->IsA( UInt8Property::StaticClass() ) )
 		{
 			UInt8Property *NumericProp = CastChecked<UInt8Property>(Prop);
-			int32 ArchetypeVal = NumericProp->GetPropertyValue_InContainer(ArchetypeContainerPtr);
-			NumericProp->SetPropertyValue_InContainer(TargetContainerPtr, ArchetypeVal);
+			int32 ArchetypeVal = NumericProp->GetPropertyValue_InContainer(ArchetypeContainerPtr, ArrayIndex);
+			NumericProp->SetPropertyValue_InContainer(TargetContainerPtr, ArchetypeVal, ArrayIndex);
 		}
 		else if ( Prop->IsA( UInt16Property::StaticClass() ) )
 		{
 			UInt16Property *NumericProp = CastChecked<UInt16Property>(Prop);
-			int16 ArchetypeVal = NumericProp->GetPropertyValue_InContainer(ArchetypeContainerPtr);
-			NumericProp->SetPropertyValue_InContainer(TargetContainerPtr, ArchetypeVal);
+			int16 ArchetypeVal = NumericProp->GetPropertyValue_InContainer(ArchetypeContainerPtr, ArrayIndex);
+			NumericProp->SetPropertyValue_InContainer(TargetContainerPtr, ArchetypeVal, ArrayIndex);
 		}
 		else if ( Prop->IsA( UIntProperty::StaticClass() ) )
 		{
 			UIntProperty *NumericProp = CastChecked<UIntProperty>(Prop);
-			int32 ArchetypeVal = NumericProp->GetPropertyValue_InContainer(ArchetypeContainerPtr);
-			NumericProp->SetPropertyValue_InContainer(TargetContainerPtr, ArchetypeVal);
+			int32 ArchetypeVal = NumericProp->GetPropertyValue_InContainer(ArchetypeContainerPtr, ArrayIndex);
+			NumericProp->SetPropertyValue_InContainer(TargetContainerPtr, ArchetypeVal, ArrayIndex);
 		}
 		else if ( Prop->IsA( UInt64Property::StaticClass() ) )
 		{
 			UInt64Property *NumericProp = CastChecked<UInt64Property>(Prop);
-			int64 ArchetypeVal = NumericProp->GetPropertyValue_InContainer(ArchetypeContainerPtr);
-			NumericProp->SetPropertyValue_InContainer(TargetContainerPtr, ArchetypeVal);
+			int64 ArchetypeVal = NumericProp->GetPropertyValue_InContainer(ArchetypeContainerPtr, ArrayIndex);
+			NumericProp->SetPropertyValue_InContainer(TargetContainerPtr, ArchetypeVal, ArrayIndex);
 		}
 		else if ( Prop->IsA( UUInt16Property::StaticClass() ) )
 		{
 			UUInt16Property *NumericProp = CastChecked<UUInt16Property>(Prop);
-			uint16 ArchetypeVal = NumericProp->GetPropertyValue_InContainer(ArchetypeContainerPtr);
-			NumericProp->SetPropertyValue_InContainer(TargetContainerPtr, ArchetypeVal);
+			uint16 ArchetypeVal = NumericProp->GetPropertyValue_InContainer(ArchetypeContainerPtr, ArrayIndex);
+			NumericProp->SetPropertyValue_InContainer(TargetContainerPtr, ArchetypeVal, ArrayIndex);
 		}
 		else if ( Prop->IsA( UUInt32Property::StaticClass() ) )
 		{
 			UUInt32Property *NumericProp = CastChecked<UUInt32Property>(Prop);
-			uint32 ArchetypeVal = NumericProp->GetPropertyValue_InContainer(ArchetypeContainerPtr);
-			NumericProp->SetPropertyValue_InContainer(TargetContainerPtr, ArchetypeVal);
+			uint32 ArchetypeVal = NumericProp->GetPropertyValue_InContainer(ArchetypeContainerPtr, ArrayIndex);
+			NumericProp->SetPropertyValue_InContainer(TargetContainerPtr, ArchetypeVal, ArrayIndex);
 		}
 		else if ( Prop->IsA( UInt64Property::StaticClass() ) )
 		{
 			UInt64Property *NumericProp = CastChecked<UInt64Property>(Prop);
-			uint64 ArchetypeVal = NumericProp->GetPropertyValue_InContainer(ArchetypeContainerPtr);
-			NumericProp->SetPropertyValue_InContainer(TargetContainerPtr, ArchetypeVal);
+			uint64 ArchetypeVal = NumericProp->GetPropertyValue_InContainer(ArchetypeContainerPtr, ArrayIndex);
+			NumericProp->SetPropertyValue_InContainer(TargetContainerPtr, ArchetypeVal, ArrayIndex);
 		}
 		else if ( Prop->IsA( UFloatProperty::StaticClass() ) )
 		{
 			UFloatProperty *NumericProp = CastChecked<UFloatProperty>(Prop);
-			float ArchetypeVal = NumericProp->GetPropertyValue_InContainer(ArchetypeContainerPtr);
-			NumericProp->SetPropertyValue_InContainer(TargetContainerPtr, ArchetypeVal);
+			float ArchetypeVal = NumericProp->GetPropertyValue_InContainer(ArchetypeContainerPtr, ArrayIndex);
+			NumericProp->SetPropertyValue_InContainer(TargetContainerPtr, ArchetypeVal, ArrayIndex);
 		}
 		else if ( Prop->IsA( UDoubleProperty::StaticClass() ) )
 		{
 			UDoubleProperty *NumericProp = CastChecked<UDoubleProperty>(Prop);
-			double ArchetypeVal = NumericProp->GetPropertyValue_InContainer(ArchetypeContainerPtr);
-			NumericProp->SetPropertyValue_InContainer(TargetContainerPtr, ArchetypeVal);
+			double ArchetypeVal = NumericProp->GetPropertyValue_InContainer(ArchetypeContainerPtr, ArrayIndex);
+			NumericProp->SetPropertyValue_InContainer(TargetContainerPtr, ArchetypeVal, ArrayIndex);
 		}
 	}
 
@@ -391,13 +406,14 @@ void ULiveEditorKismetLibrary::ReplicateChangesToChildren( FName PropertyName, U
 	}
 
 	void *ContainerPtr = Archetype;
-	UProperty *Prop = nLiveEditorKismetLibrary::GetPropertyByNameRecurse( Archetype->GetClass(), PropertyName.ToString(), &ContainerPtr );
+	int32 ArrayIndex = 0;
+	UProperty *Prop = nLiveEditorKismetLibrary::GetPropertyByNameRecurse( Archetype->GetClass(), PropertyName.ToString(), &ContainerPtr, ArrayIndex );
 	if ( Prop && Prop->IsA( UNumericProperty::StaticClass() ) )
 	{
 		check( ContainerPtr != NULL );
 		FString ClassName = Archetype->GetClass()->GetName();
 		FString ValueString;
-		void *Value = Prop->ContainerPtrToValuePtr<void>(ContainerPtr);
+		void *Value = Prop->ContainerPtrToValuePtr<void>(ContainerPtr, ArrayIndex);
 		Prop->ExportTextItem(ValueString, Value, NULL, NULL, 0);
 		FLiveEditorManager::Get().BroadcastValueUpdate( ClassName, PropertyName.ToString(), ValueString );
 	}
diff --git a/Engine/Plugins/Experimental/LiveEditor/Source/LiveEditor/Private/SGraphPinLiveEditVar.cpp b/Engine/Plugins/Experimental/LiveEditor/Source/LiveEditor/Private/SGraphPinLiveEditVar.cpp
--- a/Engine/Plugins/Experimental/LiveEditor/Source/LiveEditor/Private/SGraphPinLiveEditVar.cpp
+++ b/Engine/Plugins/Experimental/LiveEditor/Source/LiveEditor/Private/SGraphPinLiveEditVar.cpp
@@ -98,9 +98,14 @@ bool SGraphPinLiveEditVar::IsPropertyPermitedForLiveEditor( UProperty &Property
 	//NO (must not have any)
 	uint64 MustNotHaveFlags = CPF_BlueprintReadOnly | CPF_DisableEditOnInstance | CPF_EditConst | CPF_Parm;
 
-	return ( Property.HasAllPropertyFlags(MustHaveFlags) &&
-			 !Property.HasAnyPropertyFlags(MustNotHaveFlags) && 
-			 !bIsDelegate );
+	if ( !Property.HasAllPropertyFlags(MustHaveFlags) )
+		return false;
+	if ( Property.HasAnyPropertyFlags(MustNotHaveFlags) )
+		return false;
+	if ( bIsDelegate )
+		return false;
+
+	return true;
 }
 
 void SGraphPinLiveEditVar::GenerateComboBoxIndexesRecurse( UStruct *InStruct, FString PropertyPrefix, TArray< TSharedPtr<int32> >& OutComboBoxIndexes )
@@ -112,22 +117,51 @@ void SGraphPinLiveEditVar::GenerateComboBoxIndexesRecurse( UStruct *InStruct, FS
 		UProperty* Property = *PropertyIt;
 		if ( !Property->IsA(UNumericProperty::StaticClass()) )
 		{
-			UStructProperty *StructProp = Cast<UStructProperty>(Property);
-			if ( StructProp != NULL )
+			if ( Property->IsA(UStructProperty::StaticClass()) )
 			{
+				UStructProperty *StructProp = Cast<UStructProperty>(Property);
 				if ( IsPropertyPermitedForLiveEditor(*StructProp) )
 				{
-					FString NewPropertyPrefix = FString::Printf( TEXT("%s%s."), *PropertyPrefix, *Property->GetName() );
-					GenerateComboBoxIndexesRecurse( StructProp->Struct, NewPropertyPrefix, OutComboBoxIndexes );
+					if ( Property->ArrayDim > 1 )
+					{
+						for ( int32 i = 0; i < Property->ArrayDim; ++i )
+						{
+							FString arrayIndex = FString::Printf( TEXT("[%d]"), i );
+							FString NewPropertyPrefix = FString::Printf( TEXT("%s%s%s."), *PropertyPrefix, *Property->GetName(), *arrayIndex );
+							GenerateComboBoxIndexesRecurse( StructProp->Struct, NewPropertyPrefix, OutComboBoxIndexes );
+						}
+					}
+					else
+					{
+						FString NewPropertyPrefix = FString::Printf( TEXT("%s%s."), *PropertyPrefix, *Property->GetName() );
+						GenerateComboBoxIndexesRecurse( StructProp->Struct, NewPropertyPrefix, OutComboBoxIndexes );
+					}
 				}
 			}
+			else if ( Property->IsA(UArrayProperty::StaticClass()) )
+			{
+				UArrayProperty *ArrayProp = Cast<UArrayProperty>(Property);
+
+			}
 			continue;
 		}
 
 		if ( IsPropertyPermitedForLiveEditor(*Property) )
 		{
-			FString Name = PropertyPrefix + Property->GetName();
-			VariableNameList.Add( Name );
+			if ( Property->ArrayDim > 1 )
+			{
+				for ( int32 i = 0; i < Property->ArrayDim; ++i )
+				{
+					FString arrayIndex = FString::Printf( TEXT("[%d]"), i );
+					FString Name = PropertyPrefix + Property->GetName() + arrayIndex;
+					VariableNameList.Add( Name );
+				}
+			}
+			else
+			{
+				FString Name = PropertyPrefix + Property->GetName();
+				VariableNameList.Add( Name );
+			}
 		}
 	}
 }
