Fixes for LiveEditor tuning PIE Component values
Fixes to allow remote game arrays to be tuned

diff --git a/Engine/Plugins/Experimental/LiveEditor/Source/LiveEditor/Private/LiveEditorKismetLibrary.cpp b/Engine/Plugins/Experimental/LiveEditor/Source/LiveEditor/Private/LiveEditorKismetLibrary.cpp
--- a/Engine/Plugins/Experimental/LiveEditor/Source/LiveEditor/Private/LiveEditorKismetLibrary.cpp
+++ b/Engine/Plugins/Experimental/LiveEditor/Source/LiveEditor/Private/LiveEditorKismetLibrary.cpp
@@ -104,14 +104,9 @@ namespace nLiveEditorKismetLibrary
 
 		return NULL;
 	}
-	static void ModifyPropertyValue( UObject *Target, const FString &PropertyName, TEnumAsByte<ELiveEditControllerType::Type> ControlType, float Delta, int32 MidiValue, bool bShouldClamp, float ClampMin, float ClampMax )
+	static UProperty *GetPropertyByName( UObject *Target, UStruct *InStruct, const FString &PropertyName, void ** hContainerPtr, int32 &OutArrayIndex )
 	{
-		if ( Target == NULL )
-			return;
-
-		void *ContainerPtr = Target;
-		int32 ArrayIndex = 0;
-		UProperty *Prop = GetPropertyByNameRecurse( Target->GetClass(), PropertyName, &ContainerPtr, ArrayIndex );
+		UProperty *Prop = GetPropertyByNameRecurse( Target->GetClass(), PropertyName, hContainerPtr, OutArrayIndex );
 		if ( Prop == NULL )
 		{
 			AActor *AsActor = Cast<AActor>(Target);
@@ -137,12 +132,23 @@ namespace nLiveEditorKismetLibrary
 						}
 						*/
 
-						ContainerPtr = Component;
-						Prop = GetPropertyByNameRecurse( Component->GetClass(), ComponentPropertyName, &ContainerPtr, ArrayIndex );
+						*hContainerPtr = Component;
+						Prop = GetPropertyByNameRecurse( Component->GetClass(), ComponentPropertyName, hContainerPtr, OutArrayIndex );
 					}
 				}
 			}
 		}
+
+		return Prop;
+	}
+	static void ModifyPropertyValue( UObject *Target, const FString &PropertyName, TEnumAsByte<ELiveEditControllerType::Type> ControlType, float Delta, int32 MidiValue, bool bShouldClamp, float ClampMin, float ClampMax )
+	{
+		if ( Target == NULL )
+			return;
+
+		void *ContainerPtr = Target;
+		int32 ArrayIndex = 0;
+		UProperty *Prop = GetPropertyByName( Target, Target->GetClass(), PropertyName, &ContainerPtr, ArrayIndex );
 		if ( Prop == NULL || !Prop->IsA( UNumericProperty::StaticClass() ) )
 			return;
 
@@ -239,72 +245,77 @@ namespace nLiveEditorKismetLibrary
 			return;
 
 		void *ArchetypeContainerPtr = Archetype;
-		int32 ArrayIndex;
-		UProperty *Prop = GetPropertyByNameRecurse( Target->GetClass(), PropertyName.ToString(), &ArchetypeContainerPtr, ArrayIndex );
-		if ( !Prop->IsA( UNumericProperty::StaticClass() ) )
+		int32 ArchetypeArrayIndex;
+		UProperty *Prop = GetPropertyByName( Archetype, Archetype->GetClass(), PropertyName.ToString(), &ArchetypeContainerPtr, ArchetypeArrayIndex );
+		if ( !Prop || !Prop->IsA( UNumericProperty::StaticClass() ) )
 			return;
 		check(ArchetypeContainerPtr != NULL);
-		void *TargetContainerPtr = (void*)((uint8*)ArchetypeContainerPtr - (uint8*)Archetype + (uint8*)Target);
+
+		void *TargetContainerPtr = Target; //(void*)((uint8*)ArchetypeContainerPtr - (uint8*)Archetype + (uint8*)Target);
+		int32 TargetArrayIndex;
+		UProperty *TargetProp = GetPropertyByName( Target, Target->GetClass(), PropertyName.ToString(), &TargetContainerPtr, TargetArrayIndex );
+		check( TargetProp != NULL && TargetProp->IsA( UNumericProperty::StaticClass() ) );
+		check( TargetArrayIndex == ArchetypeArrayIndex );
 
 		if ( Prop->IsA( UByteProperty::StaticClass() ) )
 		{
 			UByteProperty *NumericProp = CastChecked<UByteProperty>(Prop);
-			uint8 ArchetypeVal = NumericProp->GetPropertyValue_InContainer(ArchetypeContainerPtr, ArrayIndex);
-			NumericProp->SetPropertyValue_InContainer(TargetContainerPtr, ArchetypeVal, ArrayIndex);
+			uint8 ArchetypeVal = NumericProp->GetPropertyValue_InContainer(ArchetypeContainerPtr, ArchetypeArrayIndex);
+			NumericProp->SetPropertyValue_InContainer(TargetContainerPtr, ArchetypeVal, ArchetypeArrayIndex);
 		}
 		else if ( Prop->IsA( UInt8Property::StaticClass() ) )
 		{
 			UInt8Property *NumericProp = CastChecked<UInt8Property>(Prop);
-			int32 ArchetypeVal = NumericProp->GetPropertyValue_InContainer(ArchetypeContainerPtr, ArrayIndex);
-			NumericProp->SetPropertyValue_InContainer(TargetContainerPtr, ArchetypeVal, ArrayIndex);
+			int32 ArchetypeVal = NumericProp->GetPropertyValue_InContainer(ArchetypeContainerPtr, ArchetypeArrayIndex);
+			NumericProp->SetPropertyValue_InContainer(TargetContainerPtr, ArchetypeVal, ArchetypeArrayIndex);
 		}
 		else if ( Prop->IsA( UInt16Property::StaticClass() ) )
 		{
 			UInt16Property *NumericProp = CastChecked<UInt16Property>(Prop);
-			int16 ArchetypeVal = NumericProp->GetPropertyValue_InContainer(ArchetypeContainerPtr, ArrayIndex);
-			NumericProp->SetPropertyValue_InContainer(TargetContainerPtr, ArchetypeVal, ArrayIndex);
+			int16 ArchetypeVal = NumericProp->GetPropertyValue_InContainer(ArchetypeContainerPtr, ArchetypeArrayIndex);
+			NumericProp->SetPropertyValue_InContainer(TargetContainerPtr, ArchetypeVal, ArchetypeArrayIndex);
 		}
 		else if ( Prop->IsA( UIntProperty::StaticClass() ) )
 		{
 			UIntProperty *NumericProp = CastChecked<UIntProperty>(Prop);
-			int32 ArchetypeVal = NumericProp->GetPropertyValue_InContainer(ArchetypeContainerPtr, ArrayIndex);
-			NumericProp->SetPropertyValue_InContainer(TargetContainerPtr, ArchetypeVal, ArrayIndex);
+			int32 ArchetypeVal = NumericProp->GetPropertyValue_InContainer(ArchetypeContainerPtr, ArchetypeArrayIndex);
+			NumericProp->SetPropertyValue_InContainer(TargetContainerPtr, ArchetypeVal, ArchetypeArrayIndex);
 		}
 		else if ( Prop->IsA( UInt64Property::StaticClass() ) )
 		{
 			UInt64Property *NumericProp = CastChecked<UInt64Property>(Prop);
-			int64 ArchetypeVal = NumericProp->GetPropertyValue_InContainer(ArchetypeContainerPtr, ArrayIndex);
-			NumericProp->SetPropertyValue_InContainer(TargetContainerPtr, ArchetypeVal, ArrayIndex);
+			int64 ArchetypeVal = NumericProp->GetPropertyValue_InContainer(ArchetypeContainerPtr, ArchetypeArrayIndex);
+			NumericProp->SetPropertyValue_InContainer(TargetContainerPtr, ArchetypeVal, ArchetypeArrayIndex);
 		}
 		else if ( Prop->IsA( UUInt16Property::StaticClass() ) )
 		{
 			UUInt16Property *NumericProp = CastChecked<UUInt16Property>(Prop);
-			uint16 ArchetypeVal = NumericProp->GetPropertyValue_InContainer(ArchetypeContainerPtr, ArrayIndex);
-			NumericProp->SetPropertyValue_InContainer(TargetContainerPtr, ArchetypeVal, ArrayIndex);
+			uint16 ArchetypeVal = NumericProp->GetPropertyValue_InContainer(ArchetypeContainerPtr, ArchetypeArrayIndex);
+			NumericProp->SetPropertyValue_InContainer(TargetContainerPtr, ArchetypeVal, ArchetypeArrayIndex);
 		}
 		else if ( Prop->IsA( UUInt32Property::StaticClass() ) )
 		{
 			UUInt32Property *NumericProp = CastChecked<UUInt32Property>(Prop);
-			uint32 ArchetypeVal = NumericProp->GetPropertyValue_InContainer(ArchetypeContainerPtr, ArrayIndex);
-			NumericProp->SetPropertyValue_InContainer(TargetContainerPtr, ArchetypeVal, ArrayIndex);
+			uint32 ArchetypeVal = NumericProp->GetPropertyValue_InContainer(ArchetypeContainerPtr, ArchetypeArrayIndex);
+			NumericProp->SetPropertyValue_InContainer(TargetContainerPtr, ArchetypeVal, ArchetypeArrayIndex);
 		}
 		else if ( Prop->IsA( UInt64Property::StaticClass() ) )
 		{
 			UInt64Property *NumericProp = CastChecked<UInt64Property>(Prop);
-			uint64 ArchetypeVal = NumericProp->GetPropertyValue_InContainer(ArchetypeContainerPtr, ArrayIndex);
-			NumericProp->SetPropertyValue_InContainer(TargetContainerPtr, ArchetypeVal, ArrayIndex);
+			uint64 ArchetypeVal = NumericProp->GetPropertyValue_InContainer(ArchetypeContainerPtr, ArchetypeArrayIndex);
+			NumericProp->SetPropertyValue_InContainer(TargetContainerPtr, ArchetypeVal, ArchetypeArrayIndex);
 		}
 		else if ( Prop->IsA( UFloatProperty::StaticClass() ) )
 		{
 			UFloatProperty *NumericProp = CastChecked<UFloatProperty>(Prop);
-			float ArchetypeVal = NumericProp->GetPropertyValue_InContainer(ArchetypeContainerPtr, ArrayIndex);
-			NumericProp->SetPropertyValue_InContainer(TargetContainerPtr, ArchetypeVal, ArrayIndex);
+			float ArchetypeVal = NumericProp->GetPropertyValue_InContainer(ArchetypeContainerPtr, ArchetypeArrayIndex);
+			NumericProp->SetPropertyValue_InContainer(TargetContainerPtr, ArchetypeVal, ArchetypeArrayIndex);
 		}
 		else if ( Prop->IsA( UDoubleProperty::StaticClass() ) )
 		{
 			UDoubleProperty *NumericProp = CastChecked<UDoubleProperty>(Prop);
-			double ArchetypeVal = NumericProp->GetPropertyValue_InContainer(ArchetypeContainerPtr, ArrayIndex);
-			NumericProp->SetPropertyValue_InContainer(TargetContainerPtr, ArchetypeVal, ArrayIndex);
+			double ArchetypeVal = NumericProp->GetPropertyValue_InContainer(ArchetypeContainerPtr, ArchetypeArrayIndex);
+			NumericProp->SetPropertyValue_InContainer(TargetContainerPtr, ArchetypeVal, ArchetypeArrayIndex);
 		}
 	}
 
@@ -407,7 +418,7 @@ void ULiveEditorKismetLibrary::ReplicateChangesToChildren( FName PropertyName, U
 
 	void *ContainerPtr = Archetype;
 	int32 ArrayIndex = 0;
-	UProperty *Prop = nLiveEditorKismetLibrary::GetPropertyByNameRecurse( Archetype->GetClass(), PropertyName.ToString(), &ContainerPtr, ArrayIndex );
+	UProperty *Prop = nLiveEditorKismetLibrary::GetPropertyByName( Archetype, Archetype->GetClass(), PropertyName.ToString(), &ContainerPtr, ArrayIndex );
 	if ( Prop && Prop->IsA( UNumericProperty::StaticClass() ) )
 	{
 		check( ContainerPtr != NULL );
diff --git a/Engine/Plugins/Runtime/LiveEditorListenServer/Source/LiveEditorListenServer/Private/LiveEditorListenServer.cpp b/Engine/Plugins/Runtime/LiveEditorListenServer/Source/LiveEditorListenServer/Private/LiveEditorListenServer.cpp
--- a/Engine/Plugins/Runtime/LiveEditorListenServer/Source/LiveEditorListenServer/Private/LiveEditorListenServer.cpp
+++ b/Engine/Plugins/Runtime/LiveEditorListenServer/Source/LiveEditorListenServer/Private/LiveEditorListenServer.cpp
@@ -2,9 +2,11 @@
 
 #include "LiveEditorListenServerPrivatePCH.h"
 #include "Networking.h"
+#include "DefaultValueHelper.h"
 
 namespace nLiveEditorListenServer
 {
+	/*
 	static UProperty *GetPropertyByNameRecurse( UStruct *InStruct, const FString &TokenString, void ** hContainerPtr )
 	{
 		FString FirstToken;
@@ -47,6 +49,98 @@ namespace nLiveEditorListenServer
 
 		return NULL;
 	}
+	*/
+	static UProperty *GetPropertyByNameRecurse( UStruct *InStruct, const FString &TokenString, void ** hContainerPtr, int32 &OutArrayIndex )
+	{
+		FString FirstToken;
+		FString RemainingTokens;
+		int32 SplitIndex;
+		if ( TokenString.FindChar( '.', SplitIndex ) )
+		{
+			FirstToken = TokenString.LeftChop( TokenString.Len()-SplitIndex );
+			RemainingTokens = TokenString.RightChop(SplitIndex+1);
+		}
+		else
+		{
+			FirstToken = TokenString;
+			RemainingTokens = FString(TEXT(""));
+		}
+
+		//get the array index if there is any
+		int32 ArrayIndex = 0;
+		if ( FirstToken.FindChar( '[', SplitIndex ) )
+		{
+			FString ArrayIndexString = FirstToken.RightChop( SplitIndex+1 );
+			ArrayIndexString = ArrayIndexString.LeftChop( 1 );
+			FDefaultValueHelper::ParseInt( ArrayIndexString, ArrayIndex );
+
+			FirstToken = FirstToken.LeftChop( FirstToken.Len()-SplitIndex );
+		}
+
+		for (TFieldIterator<UProperty> PropertyIt(InStruct, EFieldIteratorFlags::IncludeSuper); PropertyIt; ++PropertyIt)
+		{
+			UProperty* Property = *PropertyIt;
+			FName PropertyName = Property->GetFName();
+			if ( FirstToken == PropertyName.ToString() )
+			{
+				if ( RemainingTokens.Len() == 0 )
+				{
+					check( *hContainerPtr != NULL );
+					OutArrayIndex = ArrayIndex;
+					return Property;
+				}
+				else
+				{
+					UStructProperty *StructProp = Cast<UStructProperty>(Property);
+					if ( StructProp )
+					{
+						check( *hContainerPtr != NULL );
+						*hContainerPtr = Property->ContainerPtrToValuePtr<void>( *hContainerPtr, ArrayIndex );
+						return GetPropertyByNameRecurse( StructProp->Struct, RemainingTokens, hContainerPtr, OutArrayIndex );
+					}
+				}
+			}
+		}
+
+		return NULL;
+	}
+	static UProperty *GetPropertyByName( UObject *Target, UStruct *InStruct, const FString &PropertyName, void ** hContainerPtr, int32 &OutArrayIndex )
+	{
+		UProperty *Prop = GetPropertyByNameRecurse( Target->GetClass(), PropertyName, hContainerPtr, OutArrayIndex );
+		if ( Prop == NULL )
+		{
+			AActor *AsActor = Cast<AActor>(Target);
+			if ( AsActor != NULL )
+			{
+				FString ComponentPropertyName = PropertyName;
+				int32 SplitIndex = 0;
+				if ( ComponentPropertyName.FindChar( '.', SplitIndex ) )
+				{
+					//FString ComponentName = ComponentPropertyName.LeftChop(SplitIndex);
+					ComponentPropertyName = ComponentPropertyName.RightChop(SplitIndex+1);
+
+					TArray<UActorComponent*> ActorComponents;
+					AsActor->GetComponents(ActorComponents);
+					for ( TArray<UActorComponent*>::TIterator ComponentIt(ActorComponents); ComponentIt && !Prop; ++ComponentIt )
+					{
+						UActorComponent *Component = *ComponentIt;
+						check( Component != NULL );
+						/*
+						if ( Component->GetName() != ComponentName )
+						{
+							continue;
+						}
+						*/
+
+						*hContainerPtr = Component;
+						Prop = GetPropertyByNameRecurse( Component->GetClass(), ComponentPropertyName, hContainerPtr, OutArrayIndex );
+					}
+				}
+			}
+		}
+
+		return Prop;
+	}
 }
 
 DEFINE_LOG_CATEGORY_STATIC(LiveEditorListenServer, Log, All);
@@ -211,7 +305,8 @@ static void SetPropertyValue( UObject *Target, const FString &PropertyName, cons
 		return;
 
 	void *ContainerPtr = Target;
-	UProperty *Prop = nLiveEditorListenServer::GetPropertyByNameRecurse( Target->GetClass(), PropertyName, &ContainerPtr );
+	int32 ArrayIndex;
+	UProperty *Prop = nLiveEditorListenServer::GetPropertyByName( Target, Target->GetClass(), PropertyName, &ContainerPtr, ArrayIndex );
 	if ( !Prop || !Prop->IsA( UNumericProperty::StaticClass() ) )
 	{
 		return;
@@ -223,61 +318,61 @@ static void SetPropertyValue( UObject *Target, const FString &PropertyName, cons
 	{
 		UByteProperty *NumericProp = CastChecked<UByteProperty>(Prop);
 		uint8 Value = FCString::Atoi( *PropertyValue );
-		NumericProp->SetPropertyValue_InContainer(ContainerPtr, Value);
+		NumericProp->SetPropertyValue_InContainer(ContainerPtr, Value, ArrayIndex);
 	}
 	else if ( Prop->IsA( UInt8Property::StaticClass() ) )
 	{
 		UInt8Property *NumericProp = CastChecked<UInt8Property>(Prop);
 		int32 Value = FCString::Atoi( *PropertyValue );
-		NumericProp->SetPropertyValue_InContainer(ContainerPtr, Value);
+		NumericProp->SetPropertyValue_InContainer(ContainerPtr, Value, ArrayIndex);
 	}
 	else if ( Prop->IsA( UInt16Property::StaticClass() ) )
 	{
 		UInt16Property *NumericProp = CastChecked<UInt16Property>(Prop);
 		int16 Value = FCString::Atoi( *PropertyValue );
-		NumericProp->SetPropertyValue_InContainer(ContainerPtr, Value);
+		NumericProp->SetPropertyValue_InContainer(ContainerPtr, Value, ArrayIndex);
 	}
 	else if ( Prop->IsA( UIntProperty::StaticClass() ) )
 	{
 		UIntProperty *NumericProp = CastChecked<UIntProperty>(Prop);
 		int32 Value = FCString::Atoi( *PropertyValue );
-		NumericProp->SetPropertyValue_InContainer(ContainerPtr, Value);
+		NumericProp->SetPropertyValue_InContainer(ContainerPtr, Value, ArrayIndex);
 	}
 	else if ( Prop->IsA( UInt64Property::StaticClass() ) )
 	{
 		UInt64Property *NumericProp = CastChecked<UInt64Property>(Prop);
 		int64 Value = FCString::Atoi64( *PropertyValue );
-		NumericProp->SetPropertyValue_InContainer(ContainerPtr, Value);
+		NumericProp->SetPropertyValue_InContainer(ContainerPtr, Value, ArrayIndex);
 	}
 	else if ( Prop->IsA( UUInt16Property::StaticClass() ) )
 	{
 		UUInt16Property *NumericProp = CastChecked<UUInt16Property>(Prop);
 		uint16 Value = FCString::Atoi( *PropertyValue );
-		NumericProp->SetPropertyValue_InContainer(ContainerPtr, Value);
+		NumericProp->SetPropertyValue_InContainer(ContainerPtr, Value, ArrayIndex);
 	}
 	else if ( Prop->IsA( UUInt32Property::StaticClass() ) )
 	{
 		UUInt32Property *NumericProp = CastChecked<UUInt32Property>(Prop);
 		uint32 Value = FCString::Atoi( *PropertyValue );
-		NumericProp->SetPropertyValue_InContainer(ContainerPtr, Value);
+		NumericProp->SetPropertyValue_InContainer(ContainerPtr, Value, ArrayIndex);
 	}
 	else if ( Prop->IsA( UInt64Property::StaticClass() ) )
 	{
 		UInt64Property *NumericProp = CastChecked<UInt64Property>(Prop);
 		uint64 Value = FCString::Atoi64( *PropertyValue );
-		NumericProp->SetPropertyValue_InContainer(ContainerPtr, Value);
+		NumericProp->SetPropertyValue_InContainer(ContainerPtr, Value, ArrayIndex);
 	}
 	else if ( Prop->IsA( UFloatProperty::StaticClass() ) )
 	{
 		UFloatProperty *NumericProp = CastChecked<UFloatProperty>(Prop);
 		float Value = FCString::Atof( *PropertyValue );
-		NumericProp->SetPropertyValue_InContainer(ContainerPtr, Value);
+		NumericProp->SetPropertyValue_InContainer(ContainerPtr, Value, ArrayIndex);
 	}
 	else if ( Prop->IsA( UDoubleProperty::StaticClass() ) )
 	{
 		UDoubleProperty *NumericProp = CastChecked<UDoubleProperty>(Prop);
 		double Value = FCString::Atod( *PropertyValue );
-		NumericProp->SetPropertyValue_InContainer(ContainerPtr, Value);
+		NumericProp->SetPropertyValue_InContainer(ContainerPtr, Value, ArrayIndex);
 	}
 }
 
