Fix project generation under Rocket builds. ThirdParty source libraries are not required for SpeedTree and portmidi (since they're editor only), so their module dependences need to be added through AddThirdPartyPrivateStaticDependencies.

diff --git a/Engine/Plugins/Editor/SpeedTreeImporter/Source/SpeedTreeImporter/SpeedTreeImporter.Build.cs b/Engine/Plugins/Editor/SpeedTreeImporter/Source/SpeedTreeImporter/SpeedTreeImporter.Build.cs
--- a/Engine/Plugins/Editor/SpeedTreeImporter/Source/SpeedTreeImporter/SpeedTreeImporter.Build.cs
+++ b/Engine/Plugins/Editor/SpeedTreeImporter/Source/SpeedTreeImporter/SpeedTreeImporter.Build.cs
@@ -35,7 +35,6 @@ namespace UnrealBuildTool.Rules
                     "EditorStyle",
                     "InputCore",
 				    "RawMesh",
-				    "SpeedTree",
                     "UnrealEd",
                     "MainFrame"
 				}
@@ -46,6 +45,8 @@ namespace UnrealBuildTool.Rules
 				{
 				}
 				);
+				
+			AddThirdPartyPrivateStaticDependencies(Target, "SpeedTree");
 		}
 	}
 }
\ No newline at end of file
diff --git a/Engine/Plugins/Experimental/LiveEditor/Source/LiveEditor/LiveEditor.Build.cs b/Engine/Plugins/Experimental/LiveEditor/Source/LiveEditor/LiveEditor.Build.cs
--- a/Engine/Plugins/Experimental/LiveEditor/Source/LiveEditor/LiveEditor.Build.cs
+++ b/Engine/Plugins/Experimental/LiveEditor/Source/LiveEditor/LiveEditor.Build.cs
@@ -34,7 +34,6 @@ namespace UnrealBuildTool.Rules
                     "Kismet",
                     "KismetCompiler",
                     "GraphEditor",
-                    "portmidi",
                     "WorkspaceMenuStructure",//@TODO: This makes it depend on a module in editor, even if that module does not depend on other editor stuff!
                     "LiveEditorListenServer"
 				}
@@ -51,6 +50,8 @@ namespace UnrealBuildTool.Rules
 					// ... add any modules that your module loads dynamically here ...
 				}
 			);
+			
+			AddThirdPartyPrivateStaticDependencies(Target, "portmidi");
 		}
 	}
 }
\ No newline at end of file
