BP: ExpandNode() functions don't check "if (CompilerContext.bIsFullCompile)". The function should be called only when necessary.

diff --git a/Engine/Plugins/Experimental/LiveEditor/Source/LiveEditor/Private/K2Node_LiveEditObject.cpp b/Engine/Plugins/Experimental/LiveEditor/Source/LiveEditor/Private/K2Node_LiveEditObject.cpp
--- a/Engine/Plugins/Experimental/LiveEditor/Source/LiveEditor/Private/K2Node_LiveEditObject.cpp
+++ b/Engine/Plugins/Experimental/LiveEditor/Source/LiveEditor/Private/K2Node_LiveEditObject.cpp
@@ -254,327 +254,324 @@ void UK2Node_LiveEditObject::ExpandNode(class FKismetCompilerContext& CompilerCo
 {
 	Super::ExpandNode(CompilerContext, SourceGraph);
 
-	if ( CompilerContext.bIsFullCompile )
+	const UEdGraphSchema_K2* Schema = CompilerContext.GetSchema();
+
+	UEdGraphPin *SourceExecPin = GetExecPin();
+	UEdGraphPin *SourceThenPin = GetThenPin();
+	UEdGraphPin *SourceBlueprintPin = GetBlueprintPin();
+	UEdGraphPin *SourceBaseClassPin = GetBaseClassPin();
+	UEdGraphPin *SourceDescriptionPin = GetDescriptionPin();
+	UEdGraphPin *SourcePermittedBindingsPin = GetPermittedBindingsPin();
+	UEdGraphPin *SourceOnMidiInputPin = GetOnMidiInputPin();
+
+	UEdGraphPin *SourceVariablePin = GetVariablePin();
+	if(NULL == SourceVariablePin)
 	{
-		const UEdGraphSchema_K2* Schema = CompilerContext.GetSchema();
-
-		UEdGraphPin *SourceExecPin = GetExecPin();
-		UEdGraphPin *SourceThenPin = GetThenPin();
-		UEdGraphPin *SourceBlueprintPin = GetBlueprintPin();
-		UEdGraphPin *SourceBaseClassPin = GetBaseClassPin();
-		UEdGraphPin *SourceDescriptionPin = GetDescriptionPin();
-		UEdGraphPin *SourcePermittedBindingsPin = GetPermittedBindingsPin();
-		UEdGraphPin *SourceOnMidiInputPin = GetOnMidiInputPin();
-
-		UEdGraphPin *SourceVariablePin = GetVariablePin();
-		if(NULL == SourceVariablePin)
-		{
-			CompilerContext.MessageLog.Error(*LOCTEXT("LiveEditObjectNodeMissingBlueprint_Error", "LiveEdit node @@ must have a blueprint specified and a variable selected to tune.").ToString(), this);
-			// we break exec links so this is the only error we get, don't want the SpawnActor node being considered and giving 'unexpected node' type warnings
-			BreakAllNodeLinks();
-			return;
-		}
+		CompilerContext.MessageLog.Error(*LOCTEXT("LiveEditObjectNodeMissingBlueprint_Error", "LiveEdit node @@ must have a blueprint specified and a variable selected to tune.").ToString(), this);
+		// we break exec links so this is the only error we get, don't want the SpawnActor node being considered and giving 'unexpected node' type warnings
+		BreakAllNodeLinks();
+		return;
+	}
 
-		UClass* SpawnClass = GetClassToSpawn();
-		if(NULL == SpawnClass)
-		{
-			CompilerContext.MessageLog.Error(*LOCTEXT("LiveEditObjectNodeMissingBaseClass_Error", "LiveEdit node @@ must have a Base Class specified.").ToString(), this);
-			// we break exec links so this is the only error we get, don't want the SpawnActor node being considered and giving 'unexpected node' type warnings
-			BreakAllNodeLinks();
-			return;
-		}
+	UClass* SpawnClass = GetClassToSpawn();
+	if(NULL == SpawnClass)
+	{
+		CompilerContext.MessageLog.Error(*LOCTEXT("LiveEditObjectNodeMissingBaseClass_Error", "LiveEdit node @@ must have a Base Class specified.").ToString(), this);
+		// we break exec links so this is the only error we get, don't want the SpawnActor node being considered and giving 'unexpected node' type warnings
+		BreakAllNodeLinks();
+		return;
+	}
+
+	if ( SourcePermittedBindingsPin->LinkedTo.Num() == 0 )
+	{
+		CompilerContext.MessageLog.Error(*LOCTEXT("LiveEditObjectNodeMissingBinding_Error", "LiveEdit node @@ must specify Permitted Bindings.").ToString(), this);
+		// we break exec links so this is the only error we get, don't want the SpawnActor node being considered and giving 'unexpected node' type warnings
+		BreakAllNodeLinks();
+		return;
+	}
 
-		if ( SourcePermittedBindingsPin->LinkedTo.Num() == 0 )
+	//sanity check the VariablePin value
+	{
+		UProperty *Property = UK2Node_LiveEditObjectStatics::GetPropertyByName( SpawnClass, *SourceVariablePin->DefaultValue );
+		if ( Property == NULL || !Property->IsA(UNumericProperty::StaticClass()) )
 		{
-			CompilerContext.MessageLog.Error(*LOCTEXT("LiveEditObjectNodeMissingBinding_Error", "LiveEdit node @@ must specify Permitted Bindings.").ToString(), this);
+			CompilerContext.MessageLog.Error(*LOCTEXT("LiveEditObjectNodeInvalidVariable_Error", "LiveEdit node @@ must have a valid variable selected.").ToString(), this);
 			// we break exec links so this is the only error we get, don't want the SpawnActor node being considered and giving 'unexpected node' type warnings
 			BreakAllNodeLinks();
 			return;
 		}
+	}
 
-		//sanity check the VariablePin value
-		{
-			UProperty *Property = UK2Node_LiveEditObjectStatics::GetPropertyByName( SpawnClass, *SourceVariablePin->DefaultValue );
-			if ( Property == NULL || !Property->IsA(UNumericProperty::StaticClass()) )
-			{
-				CompilerContext.MessageLog.Error(*LOCTEXT("LiveEditObjectNodeInvalidVariable_Error", "LiveEdit node @@ must have a valid variable selected.").ToString(), this);
-				// we break exec links so this is the only error we get, don't want the SpawnActor node being considered and giving 'unexpected node' type warnings
-				BreakAllNodeLinks();
-				return;
-			}
-		}
-
-		//hooks to pins that are generated after a BaseClass is set
-		UEdGraphPin *DeltaMultPin = GetDeltaMultPin();
-		UEdGraphPin *ShouldClampPin = GetShouldClampPin();
-		UEdGraphPin *ClampMinPin = GetClampMinPin();
-		UEdGraphPin *ClampMaxPin = GetClampMaxPin();
+	//hooks to pins that are generated after a BaseClass is set
+	UEdGraphPin *DeltaMultPin = GetDeltaMultPin();
+	UEdGraphPin *ShouldClampPin = GetShouldClampPin();
+	UEdGraphPin *ClampMinPin = GetClampMinPin();
+	UEdGraphPin *ClampMaxPin = GetClampMaxPin();
 
-		UK2Node_Self *SelfNode  = CompilerContext.SpawnIntermediateNode<UK2Node_Self>(this,SourceGraph);
-		SelfNode->AllocateDefaultPins();
-		UEdGraphPin *SelfNodeThenPin = SelfNode->FindPinChecked(Schema->PN_Self);
+	UK2Node_Self *SelfNode  = CompilerContext.SpawnIntermediateNode<UK2Node_Self>(this,SourceGraph);
+	SelfNode->AllocateDefaultPins();
+	UEdGraphPin *SelfNodeThenPin = SelfNode->FindPinChecked(Schema->PN_Self);
 
-		FString EventNameGuid = GetEventName();
+	FString EventNameGuid = GetEventName();
 		
-		//Create the registration part of the LiveEditor binding process
-		{
-			UK2Node_CallFunction *RegisterForMIDINode = CompilerContext.SpawnIntermediateNode<UK2Node_CallFunction>(this,SourceGraph);
-			RegisterForMIDINode->FunctionReference.SetExternalMember( TEXT("RegisterForLiveEditEvent"), ULiveEditorKismetLibrary::StaticClass() );
-			RegisterForMIDINode->AllocateDefaultPins();
+	//Create the registration part of the LiveEditor binding process
+	{
+		UK2Node_CallFunction *RegisterForMIDINode = CompilerContext.SpawnIntermediateNode<UK2Node_CallFunction>(this,SourceGraph);
+		RegisterForMIDINode->FunctionReference.SetExternalMember( TEXT("RegisterForLiveEditEvent"), ULiveEditorKismetLibrary::StaticClass() );
+		RegisterForMIDINode->AllocateDefaultPins();
 
-			UEdGraphPin *ExecPin = RegisterForMIDINode->GetExecPin();
-			CompilerContext.MovePinLinksToIntermediate(*SourceExecPin, *ExecPin);
+		UEdGraphPin *ExecPin = RegisterForMIDINode->GetExecPin();
+		CompilerContext.MovePinLinksToIntermediate(*SourceExecPin, *ExecPin);
 
-			UEdGraphPin *ThenPin = RegisterForMIDINode->GetThenPin();
-			CompilerContext.MovePinLinksToIntermediate(*SourceThenPin, *ThenPin);
+		UEdGraphPin *ThenPin = RegisterForMIDINode->GetThenPin();
+		CompilerContext.MovePinLinksToIntermediate(*SourceThenPin, *ThenPin);
 
-			UEdGraphPin *TargetPin = RegisterForMIDINode->FindPinChecked( FString(TEXT("Target")) );
-			TargetPin->MakeLinkTo(SelfNodeThenPin);
+		UEdGraphPin *TargetPin = RegisterForMIDINode->FindPinChecked( FString(TEXT("Target")) );
+		TargetPin->MakeLinkTo(SelfNodeThenPin);
 
-			UEdGraphPin *EventNamePin = RegisterForMIDINode->FindPinChecked( FString(TEXT("EventName")) );
-			EventNamePin->DefaultValue = EventNameGuid;
+		UEdGraphPin *EventNamePin = RegisterForMIDINode->FindPinChecked( FString(TEXT("EventName")) );
+		EventNamePin->DefaultValue = EventNameGuid;
 		
-			UEdGraphPin *DescriptionPin = RegisterForMIDINode->FindPinChecked( FString(TEXT("Description")) );
-			CompilerContext.CopyPinLinksToIntermediate( *SourceDescriptionPin, *DescriptionPin);
+		UEdGraphPin *DescriptionPin = RegisterForMIDINode->FindPinChecked( FString(TEXT("Description")) );
+		CompilerContext.CopyPinLinksToIntermediate( *SourceDescriptionPin, *DescriptionPin);
 
-			UEdGraphPin *PermittedBindingsPin = RegisterForMIDINode->FindPinChecked( FString(TEXT("PermittedBindings")) );
-			CompilerContext.CopyPinLinksToIntermediate( *SourcePermittedBindingsPin, *PermittedBindingsPin);
-		}
+		UEdGraphPin *PermittedBindingsPin = RegisterForMIDINode->FindPinChecked( FString(TEXT("PermittedBindings")) );
+		CompilerContext.CopyPinLinksToIntermediate( *SourcePermittedBindingsPin, *PermittedBindingsPin);
+	}
 
-		//Create the event handling part of the LiveEditor binding process
+	//Create the event handling part of the LiveEditor binding process
+	{
+		//
+		//the event itself
+		//
+		UFunction *EventMIDISignature = GetEventMIDISignature();
+		UK2Node_Event* EventNode = CompilerContext.SpawnIntermediateNode<UK2Node_Event>(this, SourceGraph);
+		check(EventNode);
+		EventNode->EventSignatureClass = Cast<UClass>(EventMIDISignature->GetOuter());
+		EventNode->EventSignatureName = EventMIDISignature->GetFName();
+		EventNode->CustomFunctionName = *EventNameGuid;
+		EventNode->bInternalEvent = true;
+		EventNode->AllocateDefaultPins();
+
+		// Cache these out because we'll connect the sequence to it
+		UEdGraphPin *EventThenPin = EventNode->FindPinChecked( Schema->PN_Then );
+		UEdGraphPin *EventDeltaPin = EventNode->FindPinChecked( FString(TEXT("Delta")) );
+		UEdGraphPin *EventMidiValuePin = EventNode->FindPinChecked( FString(TEXT("MidiValue")) );
+		UEdGraphPin *EventControlTypePin = EventNode->FindPinChecked( FString(TEXT("ControlType")) );
+
+
+		//
+		// Check if Blueprint is NULL
+		//
+		UEdGraphPin *CompareBlueprintToNullBranchThenPin = NULL;
 		{
-			//
-			//the event itself
-			//
-			UFunction *EventMIDISignature = GetEventMIDISignature();
-			UK2Node_Event* EventNode = CompilerContext.SpawnIntermediateNode<UK2Node_Event>(this, SourceGraph);
-			check(EventNode);
-			EventNode->EventSignatureClass = Cast<UClass>(EventMIDISignature->GetOuter());
-			EventNode->EventSignatureName = EventMIDISignature->GetFName();
-			EventNode->CustomFunctionName = *EventNameGuid;
-			EventNode->bInternalEvent = true;
-			EventNode->AllocateDefaultPins();
-
-			// Cache these out because we'll connect the sequence to it
-			UEdGraphPin *EventThenPin = EventNode->FindPinChecked( Schema->PN_Then );
-			UEdGraphPin *EventDeltaPin = EventNode->FindPinChecked( FString(TEXT("Delta")) );
-			UEdGraphPin *EventMidiValuePin = EventNode->FindPinChecked( FString(TEXT("MidiValue")) );
-			UEdGraphPin *EventControlTypePin = EventNode->FindPinChecked( FString(TEXT("ControlType")) );
+			UK2Node_CallFunction *CompareBlueprintToNullNode = CompilerContext.SpawnIntermediateNode<UK2Node_CallFunction>(this,SourceGraph);
+			CompareBlueprintToNullNode->FunctionReference.SetExternalMember( TEXT("NotEqual_ObjectObject"), UKismetMathLibrary::StaticClass() );
+			CompareBlueprintToNullNode->AllocateDefaultPins();
 
+			//Set A Pin to the Blueprint Pin
+			UEdGraphPin *CompareBlueprintToNullAPin = CompareBlueprintToNullNode->FindPinChecked( FString(TEXT("A")) );
+			CompilerContext.CopyPinLinksToIntermediate( *SourceBlueprintPin, *CompareBlueprintToNullAPin);
 
-			//
-			// Check if Blueprint is NULL
-			//
-			UEdGraphPin *CompareBlueprintToNullBranchThenPin = NULL;
-			{
-				UK2Node_CallFunction *CompareBlueprintToNullNode = CompilerContext.SpawnIntermediateNode<UK2Node_CallFunction>(this,SourceGraph);
-				CompareBlueprintToNullNode->FunctionReference.SetExternalMember( TEXT("NotEqual_ObjectObject"), UKismetMathLibrary::StaticClass() );
-				CompareBlueprintToNullNode->AllocateDefaultPins();
+			// hook for Compare Blueprint to NULL result
+			UEdGraphPin *CompareBlueprintToNullResultPin = CompareBlueprintToNullNode->GetReturnValuePin();
 
-				//Set A Pin to the Blueprint Pin
-				UEdGraphPin *CompareBlueprintToNullAPin = CompareBlueprintToNullNode->FindPinChecked( FString(TEXT("A")) );
-				CompilerContext.CopyPinLinksToIntermediate( *SourceBlueprintPin, *CompareBlueprintToNullAPin);
+			// Create the BRANCH that will drive the comparison
+			UK2Node_IfThenElse* CompareBlueprintToNullBranchNode = CompilerContext.SpawnIntermediateNode<UK2Node_IfThenElse>(this, SourceGraph);
+			CompareBlueprintToNullBranchNode->AllocateDefaultPins();
 
-				// hook for Compare Blueprint to NULL result
-				UEdGraphPin *CompareBlueprintToNullResultPin = CompareBlueprintToNullNode->GetReturnValuePin();
+			//hook up the condition
+			CompareBlueprintToNullResultPin->MakeLinkTo( CompareBlueprintToNullBranchNode->GetConditionPin() );
 
-				// Create the BRANCH that will drive the comparison
-				UK2Node_IfThenElse* CompareBlueprintToNullBranchNode = CompilerContext.SpawnIntermediateNode<UK2Node_IfThenElse>(this, SourceGraph);
-				CompareBlueprintToNullBranchNode->AllocateDefaultPins();
+			//hook event to the branck input
+			EventThenPin->MakeLinkTo( CompareBlueprintToNullBranchNode->GetExecPin() );
 
-				//hook up the condition
-				CompareBlueprintToNullResultPin->MakeLinkTo( CompareBlueprintToNullBranchNode->GetConditionPin() );
-
-				//hook event to the branck input
-				EventThenPin->MakeLinkTo( CompareBlueprintToNullBranchNode->GetExecPin() );
-
-				//cache ot the THEN pin for later linkup
-				CompareBlueprintToNullBranchThenPin = CompareBlueprintToNullBranchNode->GetThenPin();
-			}
-
-			//
-			// Get Class Default Object
-			//
-			UK2Node_CallFunction *GetClassDefaultObjectNode = CompilerContext.SpawnIntermediateNode<UK2Node_CallFunction>(this,SourceGraph);
-			GetClassDefaultObjectNode->FunctionReference.SetExternalMember( TEXT("GetBlueprintClassDefaultObject"), ULiveEditorKismetLibrary::StaticClass() );
-			GetClassDefaultObjectNode->AllocateDefaultPins();
-
-			UEdGraphPin *GetClassDefaultObjectBlueprintPin = GetClassDefaultObjectNode->FindPinChecked( TEXT("Blueprint") );
-			CompilerContext.CopyPinLinksToIntermediate( *SourceBlueprintPin, *GetClassDefaultObjectBlueprintPin);
+			//cache ot the THEN pin for later linkup
+			CompareBlueprintToNullBranchThenPin = CompareBlueprintToNullBranchNode->GetThenPin();
+		}
 
-			//hook for later -> the pointer to the ClassDefaultObject of our BlueprintPin
-			UEdGraphPin *GetClassDefaultObjectResultPin = GetClassDefaultObjectNode->GetReturnValuePin();
+		//
+		// Get Class Default Object
+		//
+		UK2Node_CallFunction *GetClassDefaultObjectNode = CompilerContext.SpawnIntermediateNode<UK2Node_CallFunction>(this,SourceGraph);
+		GetClassDefaultObjectNode->FunctionReference.SetExternalMember( TEXT("GetBlueprintClassDefaultObject"), ULiveEditorKismetLibrary::StaticClass() );
+		GetClassDefaultObjectNode->AllocateDefaultPins();
 
+		UEdGraphPin *GetClassDefaultObjectBlueprintPin = GetClassDefaultObjectNode->FindPinChecked( TEXT("Blueprint") );
+		CompilerContext.CopyPinLinksToIntermediate( *SourceBlueprintPin, *GetClassDefaultObjectBlueprintPin);
 
-			//
-			// Compare to BaseClass to make sure that the target Blueprint IsA(BaseClass)
-			//
-			UEdGraphPin *ClassIsChildOfBranchThenPin = NULL;
-			{
-				//
-				//we need to get the class of the Blueprint pin
-				UK2Node_CallFunction *GetClassNode = CompilerContext.SpawnIntermediateNode<UK2Node_CallFunction>(this,SourceGraph);
-				GetClassNode->FunctionReference.SetExternalMember( TEXT("GetObjectClass"), UGameplayStatics::StaticClass() );
-				GetClassNode->AllocateDefaultPins();
-
-				//Pin in the GetClassDefaultObjectResultPin to the Object Parameter of the GetObjectClass FUNCTION
-				//we want to make sure that the Class of the DEFAULT_OBJECT IsA( BaseClass )
-				UEdGraphPin *GetClassObjectPin = GetClassNode->FindPinChecked( FString(TEXT("Object")) );
-				GetClassDefaultObjectResultPin->MakeLinkTo( GetClassObjectPin );
-
-				//hook for the Class result
-				UEdGraphPin *GetClassReturnValuePin = GetClassNode->GetReturnValuePin();
-
-				//
-				//the ClassIsChildOf FUNCTION
-				UK2Node_CallFunction *ClassIsChildOfNode = CompilerContext.SpawnIntermediateNode<UK2Node_CallFunction>(this,SourceGraph);
-				ClassIsChildOfNode->FunctionReference.SetExternalMember( TEXT("ClassIsChildOf"), UKismetMathLibrary::StaticClass() );
-				ClassIsChildOfNode->AllocateDefaultPins();
-
-				//hook up the test pin
-				UEdGraphPin *ClassIsChildOfTestPin = ClassIsChildOfNode->FindPinChecked( FString(TEXT("TestClass")) );
-				GetClassReturnValuePin->MakeLinkTo( ClassIsChildOfTestPin );
-
-				//copy our BaseClass Pin into the ClassIsChildOf Parameter
-				UEdGraphPin *ClassIsChildOfParentPin = ClassIsChildOfNode->FindPinChecked( FString(TEXT("ParentClass")) );
-				CompilerContext.CopyPinLinksToIntermediate( *SourceBaseClassPin, *ClassIsChildOfParentPin);
-
-				//hook for return value
-				UEdGraphPin *ClassIsChildOfResultPin = ClassIsChildOfNode->GetReturnValuePin();
-
-				//
-				// Create the BRANCH that will drive the comparison
-				UK2Node_IfThenElse* ClassIsChildOfBranchNode = CompilerContext.SpawnIntermediateNode<UK2Node_IfThenElse>(this, SourceGraph);
-				ClassIsChildOfBranchNode->AllocateDefaultPins();
-
-				//hook up the previous branch to this one
-				check( CompareBlueprintToNullBranchThenPin != NULL );
-				CompareBlueprintToNullBranchThenPin->MakeLinkTo( ClassIsChildOfBranchNode->GetExecPin() );
-
-				//hook up our condition
-				ClassIsChildOfResultPin->MakeLinkTo( ClassIsChildOfBranchNode->GetConditionPin() );
-
-				//cache ot the THEN pin for later linkup
-				ClassIsChildOfBranchThenPin = ClassIsChildOfBranchNode->GetThenPin();
-			}
+		//hook for later -> the pointer to the ClassDefaultObject of our BlueprintPin
+		UEdGraphPin *GetClassDefaultObjectResultPin = GetClassDefaultObjectNode->GetReturnValuePin();
 
 
+		//
+		// Compare to BaseClass to make sure that the target Blueprint IsA(BaseClass)
+		//
+		UEdGraphPin *ClassIsChildOfBranchThenPin = NULL;
+		{
 			//
-			//The set variable function (to set LiveEdited new value)
-			//
-			UK2Node_CallFunction *ModifyVarNode = CompilerContext.SpawnIntermediateNode<UK2Node_CallFunction>(this,SourceGraph);
-			ModifyVarNode->FunctionReference.SetExternalMember( TEXT("ModifyPropertyByName"), ULiveEditorKismetLibrary::StaticClass() );
-			ModifyVarNode->AllocateDefaultPins();
+			//we need to get the class of the Blueprint pin
+			UK2Node_CallFunction *GetClassNode = CompilerContext.SpawnIntermediateNode<UK2Node_CallFunction>(this,SourceGraph);
+			GetClassNode->FunctionReference.SetExternalMember( TEXT("GetObjectClass"), UGameplayStatics::StaticClass() );
+			GetClassNode->AllocateDefaultPins();
 
-			// Make link from the event to the Set variable node
-			UEdGraphPin *ModifyVarExecPin = ModifyVarNode->GetExecPin();
-			ClassIsChildOfBranchThenPin->MakeLinkTo( ModifyVarExecPin );
+			//Pin in the GetClassDefaultObjectResultPin to the Object Parameter of the GetObjectClass FUNCTION
+			//we want to make sure that the Class of the DEFAULT_OBJECT IsA( BaseClass )
+			UEdGraphPin *GetClassObjectPin = GetClassNode->FindPinChecked( FString(TEXT("Object")) );
+			GetClassDefaultObjectResultPin->MakeLinkTo( GetClassObjectPin );
 
-			//link up the Target Pin
-			UEdGraphPin *ModifyVarNodeTargetPin = ModifyVarNode->FindPinChecked( TEXT("Target") );
-			GetClassDefaultObjectResultPin->MakeLinkTo( ModifyVarNodeTargetPin );
+			//hook for the Class result
+			UEdGraphPin *GetClassReturnValuePin = GetClassNode->GetReturnValuePin();
 
-			//link up the PropertyName Pin
-			UEdGraphPin *ModifyVarNodePropertyNamePin = ModifyVarNode->FindPinChecked( TEXT("PropertyName") );
-			ModifyVarNodePropertyNamePin->DefaultValue = SourceVariablePin->DefaultValue;
+			//
+			//the ClassIsChildOf FUNCTION
+			UK2Node_CallFunction *ClassIsChildOfNode = CompilerContext.SpawnIntermediateNode<UK2Node_CallFunction>(this,SourceGraph);
+			ClassIsChildOfNode->FunctionReference.SetExternalMember( TEXT("ClassIsChildOf"), UKismetMathLibrary::StaticClass() );
+			ClassIsChildOfNode->AllocateDefaultPins();
 
-			//link up the MIDI Value Pin
-			UEdGraphPin *ModifyVarNodeMidiValuePin = ModifyVarNode->FindPinChecked( TEXT("MidiValue") );
-			EventMidiValuePin->MakeLinkTo(ModifyVarNodeMidiValuePin);
+			//hook up the test pin
+			UEdGraphPin *ClassIsChildOfTestPin = ClassIsChildOfNode->FindPinChecked( FString(TEXT("TestClass")) );
+			GetClassReturnValuePin->MakeLinkTo( ClassIsChildOfTestPin );
 
-			//link up the ControlType Pin
-			UEdGraphPin *ModifyVarNodeControlTypePin = ModifyVarNode->FindPinChecked( TEXT("ControlType") );
-			EventControlTypePin->MakeLinkTo(ModifyVarNodeControlTypePin);
+			//copy our BaseClass Pin into the ClassIsChildOf Parameter
+			UEdGraphPin *ClassIsChildOfParentPin = ClassIsChildOfNode->FindPinChecked( FString(TEXT("ParentClass")) );
+			CompilerContext.CopyPinLinksToIntermediate( *SourceBaseClassPin, *ClassIsChildOfParentPin);
 
-			//hook for the Delta Pin
-			UEdGraphPin *ModifyVarNodeDeltaPin = ModifyVarNode->FindPinChecked( TEXT("Delta") );
+			//hook for return value
+			UEdGraphPin *ClassIsChildOfResultPin = ClassIsChildOfNode->GetReturnValuePin();
 
-			//Clamping
-			if ( ShouldClampPin->DefaultValue == FString(TEXT("true")) )
-			{
-				UEdGraphPin *ModifyVarNodeShouldClampPin = ModifyVarNode->FindPinChecked( TEXT("bShouldClamp") );
-				CompilerContext.CopyPinLinksToIntermediate( *ShouldClampPin, *ModifyVarNodeShouldClampPin);
+			//
+			// Create the BRANCH that will drive the comparison
+			UK2Node_IfThenElse* ClassIsChildOfBranchNode = CompilerContext.SpawnIntermediateNode<UK2Node_IfThenElse>(this, SourceGraph);
+			ClassIsChildOfBranchNode->AllocateDefaultPins();
 
-				check( ClampMinPin != NULL );
-				UEdGraphPin *ModifyVarNodeClampMinPin = ModifyVarNode->FindPinChecked( TEXT("ClampMin") );
-				CompilerContext.CopyPinLinksToIntermediate( *ClampMinPin, *ModifyVarNodeClampMinPin);
+			//hook up the previous branch to this one
+			check( CompareBlueprintToNullBranchThenPin != NULL );
+			CompareBlueprintToNullBranchThenPin->MakeLinkTo( ClassIsChildOfBranchNode->GetExecPin() );
 
-				check( ClampMaxPin != NULL );
-				UEdGraphPin *ModifyVarNodeClampMaxPin = ModifyVarNode->FindPinChecked( TEXT("ClampMax") );
-				CompilerContext.CopyPinLinksToIntermediate( *ClampMaxPin, *ModifyVarNodeClampMaxPin);
-			}
+			//hook up our condition
+			ClassIsChildOfResultPin->MakeLinkTo( ClassIsChildOfBranchNode->GetConditionPin() );
 
-			//hook for ModifyVar THEN
-			UEdGraphPin *ModifyVarNodeThenPin = ModifyVarNode->GetThenPin();
+			//cache ot the THEN pin for later linkup
+			ClassIsChildOfBranchThenPin = ClassIsChildOfBranchNode->GetThenPin();
+		}
 
-			//
-			// The Multiply Delta * DeltaMult function
-			//
-			UK2Node_CallFunction *MultiplyNode = CompilerContext.SpawnIntermediateNode<UK2Node_CallFunction>(this,SourceGraph);
-			MultiplyNode->FunctionReference.SetExternalMember( TEXT("Multiply_FloatFloat"), UKismetMathLibrary::StaticClass() );
-			MultiplyNode->AllocateDefaultPins();
 
-			//cache this out. it will be linked to from the output of the (int)Delta -> (float)Delta Conversion function
-			UEdGraphPin *MultiplyNodeFirstPin = MultiplyNode->FindPinChecked( FString(TEXT("A")) );
+		//
+		//The set variable function (to set LiveEdited new value)
+		//
+		UK2Node_CallFunction *ModifyVarNode = CompilerContext.SpawnIntermediateNode<UK2Node_CallFunction>(this,SourceGraph);
+		ModifyVarNode->FunctionReference.SetExternalMember( TEXT("ModifyPropertyByName"), ULiveEditorKismetLibrary::StaticClass() );
+		ModifyVarNode->AllocateDefaultPins();
 
-			// 2nd input to the Add function comes from the Current variable value
-			UEdGraphPin *MultiplyNodeSecondPin = MultiplyNode->FindPinChecked( FString(TEXT("B")) );
-			CompilerContext.CopyPinLinksToIntermediate( *DeltaMultPin, *MultiplyNodeSecondPin);
+		// Make link from the event to the Set variable node
+		UEdGraphPin *ModifyVarExecPin = ModifyVarNode->GetExecPin();
+		ClassIsChildOfBranchThenPin->MakeLinkTo( ModifyVarExecPin );
 
-			UEdGraphPin *MultiplyNodeReturnValuePin = MultiplyNode->GetReturnValuePin();
-			MultiplyNodeReturnValuePin->MakeLinkTo( ModifyVarNodeDeltaPin );
+		//link up the Target Pin
+		UEdGraphPin *ModifyVarNodeTargetPin = ModifyVarNode->FindPinChecked( TEXT("Target") );
+		GetClassDefaultObjectResultPin->MakeLinkTo( ModifyVarNodeTargetPin );
 
-			//
-			// The Convert function to go from (int)Delta to ULiveEditorKismetLibrary::ModifyPropertyByName(... float Delta ...)
-			//
-			FName ConvertFunctionName;
-			bool success = Schema->SearchForAutocastFunction( EventDeltaPin, MultiplyNodeFirstPin, ConvertFunctionName );
-			check( success );
-			UK2Node_CallFunction *ConvertDeltaNode = CompilerContext.SpawnIntermediateNode<UK2Node_CallFunction>(this,SourceGraph);
-			ConvertDeltaNode->FunctionReference.SetExternalMember( ConvertFunctionName, UKismetMathLibrary::StaticClass() );
-			ConvertDeltaNode->AllocateDefaultPins();
-
-			FName PinName;
-			success = UK2Node_LiveEditObjectStatics::SearchForConvertPinName( Schema, EventDeltaPin, PinName );
-			check( success );
-			UEdGraphPin *ConvertDeltaInputPin = ConvertDeltaNode->FindPinChecked( PinName.ToString() );
-			EventDeltaPin->MakeLinkTo( ConvertDeltaInputPin );
-
-			UEdGraphPin *ConvertDeltaOutputPin = ConvertDeltaNode->GetReturnValuePin();
-			ConvertDeltaOutputPin->MakeLinkTo( MultiplyNodeFirstPin );
+		//link up the PropertyName Pin
+		UEdGraphPin *ModifyVarNodePropertyNamePin = ModifyVarNode->FindPinChecked( TEXT("PropertyName") );
+		ModifyVarNodePropertyNamePin->DefaultValue = SourceVariablePin->DefaultValue;
 
-			//
-			// TODO - markDirty
-			//
+		//link up the MIDI Value Pin
+		UEdGraphPin *ModifyVarNodeMidiValuePin = ModifyVarNode->FindPinChecked( TEXT("MidiValue") );
+		EventMidiValuePin->MakeLinkTo(ModifyVarNodeMidiValuePin);
 
-			//
-			// send out the object value updates
-			//
-			UK2Node_CallFunction *ReplicationNode = CompilerContext.SpawnIntermediateNode<UK2Node_CallFunction>(this,SourceGraph);
-			ReplicationNode->FunctionReference.SetExternalMember( TEXT("ReplicateChangesToChildren"), ULiveEditorKismetLibrary::StaticClass() );
-			ReplicationNode->AllocateDefaultPins();
+		//link up the ControlType Pin
+		UEdGraphPin *ModifyVarNodeControlTypePin = ModifyVarNode->FindPinChecked( TEXT("ControlType") );
+		EventControlTypePin->MakeLinkTo(ModifyVarNodeControlTypePin);
 
-			UEdGraphPin *ReplicationNodeVarNamePin = ReplicationNode->FindPinChecked( TEXT("PropertyName") );
-			ReplicationNodeVarNamePin->DefaultValue = SourceVariablePin->DefaultValue;
+		//hook for the Delta Pin
+		UEdGraphPin *ModifyVarNodeDeltaPin = ModifyVarNode->FindPinChecked( TEXT("Delta") );
 
-			UEdGraphPin *ReplicationNodeArchetypePin = ReplicationNode->FindPinChecked( FString(TEXT("Archetype")) );
-			GetClassDefaultObjectResultPin->MakeLinkTo( ReplicationNodeArchetypePin );
-
-			UEdGraphPin *ReplicationNodeExecPin = ReplicationNode->GetExecPin();
-			ModifyVarNodeThenPin->MakeLinkTo( ReplicationNodeExecPin );
+		//Clamping
+		if ( ShouldClampPin->DefaultValue == FString(TEXT("true")) )
+		{
+			UEdGraphPin *ModifyVarNodeShouldClampPin = ModifyVarNode->FindPinChecked( TEXT("bShouldClamp") );
+			CompilerContext.CopyPinLinksToIntermediate( *ShouldClampPin, *ModifyVarNodeShouldClampPin);
 
-			UEdGraphPin *ReplicationNodeThenPin = ReplicationNode->FindPinChecked( FString(TEXT("then")) );
+			check( ClampMinPin != NULL );
+			UEdGraphPin *ModifyVarNodeClampMinPin = ModifyVarNode->FindPinChecked( TEXT("ClampMin") );
+			CompilerContext.CopyPinLinksToIntermediate( *ClampMinPin, *ModifyVarNodeClampMinPin);
 
-			//
-			// Finally, activate our OnMidiInput pin
-			//
-			CompilerContext.CopyPinLinksToIntermediate( *SourceOnMidiInputPin, *ReplicationNodeThenPin);
-			
+			check( ClampMaxPin != NULL );
+			UEdGraphPin *ModifyVarNodeClampMaxPin = ModifyVarNode->FindPinChecked( TEXT("ClampMax") );
+			CompilerContext.CopyPinLinksToIntermediate( *ClampMaxPin, *ModifyVarNodeClampMaxPin);
 		}
 
-		// Break any links to the expanded node
-		BreakAllNodeLinks();
+		//hook for ModifyVar THEN
+		UEdGraphPin *ModifyVarNodeThenPin = ModifyVarNode->GetThenPin();
+
+		//
+		// The Multiply Delta * DeltaMult function
+		//
+		UK2Node_CallFunction *MultiplyNode = CompilerContext.SpawnIntermediateNode<UK2Node_CallFunction>(this,SourceGraph);
+		MultiplyNode->FunctionReference.SetExternalMember( TEXT("Multiply_FloatFloat"), UKismetMathLibrary::StaticClass() );
+		MultiplyNode->AllocateDefaultPins();
+
+		//cache this out. it will be linked to from the output of the (int)Delta -> (float)Delta Conversion function
+		UEdGraphPin *MultiplyNodeFirstPin = MultiplyNode->FindPinChecked( FString(TEXT("A")) );
+
+		// 2nd input to the Add function comes from the Current variable value
+		UEdGraphPin *MultiplyNodeSecondPin = MultiplyNode->FindPinChecked( FString(TEXT("B")) );
+		CompilerContext.CopyPinLinksToIntermediate( *DeltaMultPin, *MultiplyNodeSecondPin);
+
+		UEdGraphPin *MultiplyNodeReturnValuePin = MultiplyNode->GetReturnValuePin();
+		MultiplyNodeReturnValuePin->MakeLinkTo( ModifyVarNodeDeltaPin );
+
+		//
+		// The Convert function to go from (int)Delta to ULiveEditorKismetLibrary::ModifyPropertyByName(... float Delta ...)
+		//
+		FName ConvertFunctionName;
+		bool success = Schema->SearchForAutocastFunction( EventDeltaPin, MultiplyNodeFirstPin, ConvertFunctionName );
+		check( success );
+		UK2Node_CallFunction *ConvertDeltaNode = CompilerContext.SpawnIntermediateNode<UK2Node_CallFunction>(this,SourceGraph);
+		ConvertDeltaNode->FunctionReference.SetExternalMember( ConvertFunctionName, UKismetMathLibrary::StaticClass() );
+		ConvertDeltaNode->AllocateDefaultPins();
+
+		FName PinName;
+		success = UK2Node_LiveEditObjectStatics::SearchForConvertPinName( Schema, EventDeltaPin, PinName );
+		check( success );
+		UEdGraphPin *ConvertDeltaInputPin = ConvertDeltaNode->FindPinChecked( PinName.ToString() );
+		EventDeltaPin->MakeLinkTo( ConvertDeltaInputPin );
+
+		UEdGraphPin *ConvertDeltaOutputPin = ConvertDeltaNode->GetReturnValuePin();
+		ConvertDeltaOutputPin->MakeLinkTo( MultiplyNodeFirstPin );
+
+		//
+		// TODO - markDirty
+		//
+
+		//
+		// send out the object value updates
+		//
+		UK2Node_CallFunction *ReplicationNode = CompilerContext.SpawnIntermediateNode<UK2Node_CallFunction>(this,SourceGraph);
+		ReplicationNode->FunctionReference.SetExternalMember( TEXT("ReplicateChangesToChildren"), ULiveEditorKismetLibrary::StaticClass() );
+		ReplicationNode->AllocateDefaultPins();
+
+		UEdGraphPin *ReplicationNodeVarNamePin = ReplicationNode->FindPinChecked( TEXT("PropertyName") );
+		ReplicationNodeVarNamePin->DefaultValue = SourceVariablePin->DefaultValue;
+
+		UEdGraphPin *ReplicationNodeArchetypePin = ReplicationNode->FindPinChecked( FString(TEXT("Archetype")) );
+		GetClassDefaultObjectResultPin->MakeLinkTo( ReplicationNodeArchetypePin );
+
+		UEdGraphPin *ReplicationNodeExecPin = ReplicationNode->GetExecPin();
+		ModifyVarNodeThenPin->MakeLinkTo( ReplicationNodeExecPin );
+
+		UEdGraphPin *ReplicationNodeThenPin = ReplicationNode->FindPinChecked( FString(TEXT("then")) );
+
+		//
+		// Finally, activate our OnMidiInput pin
+		//
+		CompilerContext.CopyPinLinksToIntermediate( *SourceOnMidiInputPin, *ReplicationNodeThenPin);
+			
 	}
+
+	// Break any links to the expanded node
+	BreakAllNodeLinks();
 }
 
 void UK2Node_LiveEditObject::GetNodeAttributes( TArray<TKeyValuePair<FString, FString>>& OutNodeAttributes ) const
diff --git a/Engine/Source/Editor/BlueprintGraph/Private/K2Node.cpp b/Engine/Source/Editor/BlueprintGraph/Private/K2Node.cpp
--- a/Engine/Source/Editor/BlueprintGraph/Private/K2Node.cpp
+++ b/Engine/Source/Editor/BlueprintGraph/Private/K2Node.cpp
@@ -667,17 +667,14 @@ void UK2Node::ExpandSplitPin(FKismetCompilerContext* CompilerContext, UEdGraph*
 
 void UK2Node::ExpandNode(FKismetCompilerContext& CompilerContext, UEdGraph* SourceGraph)
 {
-	if (CompilerContext.bIsFullCompile)
+	// We iterate the array in reverse so we can both remove the subpins safely after we've read them and
+	// so we have split nested structs we combine them back together in the right order
+	for (int32 PinIndex=Pins.Num() - 1; PinIndex >= 0; --PinIndex)
 	{
-		// We iterate the array in reverse so we can both remove the subpins safely after we've read them and
-		// so we have split nested structs we combine them back together in the right order
-		for (int32 PinIndex=Pins.Num() - 1; PinIndex >= 0; --PinIndex)
+		UEdGraphPin* Pin = Pins[PinIndex];
+		if (Pin->SubPins.Num() > 0)
 		{
-			UEdGraphPin* Pin = Pins[PinIndex];
-			if (Pin->SubPins.Num() > 0)
-			{
-				ExpandSplitPin(&CompilerContext, SourceGraph, Pin);
-			}
+			ExpandSplitPin(&CompilerContext, SourceGraph, Pin);
 		}
 	}
 }
diff --git a/Engine/Source/Editor/BlueprintGraph/Private/K2Node_AddComponent.cpp b/Engine/Source/Editor/BlueprintGraph/Private/K2Node_AddComponent.cpp
--- a/Engine/Source/Editor/BlueprintGraph/Private/K2Node_AddComponent.cpp
+++ b/Engine/Source/Editor/BlueprintGraph/Private/K2Node_AddComponent.cpp
@@ -366,34 +366,30 @@ void UK2Node_AddComponent::ExpandNode(class FKismetCompilerContext& CompilerCont
 {
 	Super::ExpandNode(CompilerContext, SourceGraph);
 
-	
-	if (CompilerContext.bIsFullCompile)
+	auto TransformPin = GetRelativeTransformPin();
+	if (TransformPin && !TransformPin->LinkedTo.Num())
 	{
-		auto TransformPin = GetRelativeTransformPin();
-		if (TransformPin && !TransformPin->LinkedTo.Num())
+		FString DefaultValue;
+
+		// Try and find the template and get relative transform from it
+		UEdGraphPin* TemplateNamePin = GetTemplateNamePinChecked();
+		const FString TemplateName = TemplateNamePin->DefaultValue;
+		check(CompilerContext.Blueprint);
+		USceneComponent* SceneCompTemplate = Cast<USceneComponent>(CompilerContext.Blueprint->FindTemplateByName(FName(*TemplateName)));
+		if (SceneCompTemplate)
 		{
-			FString DefaultValue;
-
-			// Try and find the template and get relative transform from it
-			UEdGraphPin* TemplateNamePin = GetTemplateNamePinChecked();
-			const FString TemplateName = TemplateNamePin->DefaultValue;
-			check(CompilerContext.Blueprint);
-			USceneComponent* SceneCompTemplate = Cast<USceneComponent>(CompilerContext.Blueprint->FindTemplateByName(FName(*TemplateName)));
-			if (SceneCompTemplate)
-			{
-				FTransform TemplateTransform = FTransform(SceneCompTemplate->RelativeRotation, SceneCompTemplate->RelativeLocation, SceneCompTemplate->RelativeScale3D);
-				DefaultValue = TemplateTransform.ToString();
-			}
+			FTransform TemplateTransform = FTransform(SceneCompTemplate->RelativeRotation, SceneCompTemplate->RelativeLocation, SceneCompTemplate->RelativeScale3D);
+			DefaultValue = TemplateTransform.ToString();
+		}
 
-			auto ValuePin = InnerHandleAutoCreateRef(this, TransformPin, CompilerContext, SourceGraph, !DefaultValue.IsEmpty());
-			if (ValuePin)
-			{
-				ValuePin->DefaultValue = DefaultValue;
-			}
+		auto ValuePin = InnerHandleAutoCreateRef(this, TransformPin, CompilerContext, SourceGraph, !DefaultValue.IsEmpty());
+		if (ValuePin)
+		{
+			ValuePin->DefaultValue = DefaultValue;
 		}
 	}
 
-	if (CompilerContext.bIsFullCompile && bHasExposedVariable)
+	if (bHasExposedVariable)
 	{
 		static FString ObjectParamName = FString(TEXT("Object"));
 		static FString ValueParamName = FString(TEXT("Value"));
diff --git a/Engine/Source/Editor/BlueprintGraph/Private/K2Node_BaseAsyncTask.cpp b/Engine/Source/Editor/BlueprintGraph/Private/K2Node_BaseAsyncTask.cpp
--- a/Engine/Source/Editor/BlueprintGraph/Private/K2Node_BaseAsyncTask.cpp
+++ b/Engine/Source/Editor/BlueprintGraph/Private/K2Node_BaseAsyncTask.cpp
@@ -264,10 +264,6 @@ void UK2Node_BaseAsyncTask::ExpandNode(class FKismetCompilerContext& CompilerCon
 {
     Super::ExpandNode(CompilerContext, SourceGraph);
 
-	if (!CompilerContext.bIsFullCompile)
-	{
-		return;
-	}
 	const UEdGraphSchema_K2* Schema = CompilerContext.GetSchema();
 	check(SourceGraph && Schema);
 	bool bIsErrorFree = true;
diff --git a/Engine/Source/Editor/BlueprintGraph/Private/K2Node_CallFunction.cpp b/Engine/Source/Editor/BlueprintGraph/Private/K2Node_CallFunction.cpp
--- a/Engine/Source/Editor/BlueprintGraph/Private/K2Node_CallFunction.cpp
+++ b/Engine/Source/Editor/BlueprintGraph/Private/K2Node_CallFunction.cpp
@@ -1662,196 +1662,193 @@ void UK2Node_CallFunction::ExpandNode(class FKismetCompilerContext& CompilerCont
 {
 	Super::ExpandNode(CompilerContext, SourceGraph);
 
-	if (CompilerContext.bIsFullCompile)
-	{
-		const UEdGraphSchema_K2* Schema = CompilerContext.GetSchema();
-		UFunction* Function = GetTargetFunction();
+	const UEdGraphSchema_K2* Schema = CompilerContext.GetSchema();
+	UFunction* Function = GetTargetFunction();
 
-		// connect DefaultToSelf and WorldContext inside static functions to proper 'self'  
-		const bool bInsideBpFuncLibrary = CompilerContext.Blueprint && (BPTYPE_FunctionLibrary == CompilerContext.Blueprint->BlueprintType);
-		if (bInsideBpFuncLibrary && SourceGraph && Function)
+	// connect DefaultToSelf and WorldContext inside static functions to proper 'self'  
+	const bool bInsideBpFuncLibrary = CompilerContext.Blueprint && (BPTYPE_FunctionLibrary == CompilerContext.Blueprint->BlueprintType);
+	if (bInsideBpFuncLibrary && SourceGraph && Function)
+	{
+		TArray<UK2Node_FunctionEntry*> EntryPoints;
+		SourceGraph->GetNodesOfClass(EntryPoints);
+		if (1 != EntryPoints.Num())
 		{
-			TArray<UK2Node_FunctionEntry*> EntryPoints;
-			SourceGraph->GetNodesOfClass(EntryPoints);
-			if (1 != EntryPoints.Num())
-			{
-				CompilerContext.MessageLog.Warning(*FString::Printf(*LOCTEXT("WrongEntryPointsNum", "%i entry points found while expanding node @@").ToString(), EntryPoints.Num()), this);
-			}
-			else if (auto BetterSelfPin = EntryPoints[0]->GetAutoWorldContextPin())
-			{
-				FString const DefaultToSelfMetaValue = Function->GetMetaData(FBlueprintMetadata::MD_DefaultToSelf);
-				FString const WorldContextMetaValue = Function->GetMetaData(FBlueprintMetadata::MD_WorldContext);
+			CompilerContext.MessageLog.Warning(*FString::Printf(*LOCTEXT("WrongEntryPointsNum", "%i entry points found while expanding node @@").ToString(), EntryPoints.Num()), this);
+		}
+		else if (auto BetterSelfPin = EntryPoints[0]->GetAutoWorldContextPin())
+		{
+			FString const DefaultToSelfMetaValue = Function->GetMetaData(FBlueprintMetadata::MD_DefaultToSelf);
+			FString const WorldContextMetaValue = Function->GetMetaData(FBlueprintMetadata::MD_WorldContext);
 
-				struct FStructConnectHelper
+			struct FStructConnectHelper
+			{
+				static void Connect(const FString& PinName, UK2Node* Node, UEdGraphPin* BetterSelf, const UEdGraphSchema_K2* InSchema, FCompilerResultsLog& MessageLog)
 				{
-					static void Connect(const FString& PinName, UK2Node* Node, UEdGraphPin* BetterSelf, const UEdGraphSchema_K2* InSchema, FCompilerResultsLog& MessageLog)
+					auto Pin = Node->FindPin(PinName);
+					if (!PinName.IsEmpty() && Pin && !Pin->LinkedTo.Num())
 					{
-						auto Pin = Node->FindPin(PinName);
-						if (!PinName.IsEmpty() && Pin && !Pin->LinkedTo.Num())
+						const bool bConnected = InSchema->TryCreateConnection(Pin, BetterSelf);
+						if (!bConnected)
 						{
-							const bool bConnected = InSchema->TryCreateConnection(Pin, BetterSelf);
-							if (!bConnected)
-							{
-								MessageLog.Warning(*LOCTEXT("DefaultToSelfNotConnected", "DefaultToSelf pin @@ from node @@ cannot be connected to @@").ToString(), Pin, Node, BetterSelf);
-							}
+							MessageLog.Warning(*LOCTEXT("DefaultToSelfNotConnected", "DefaultToSelf pin @@ from node @@ cannot be connected to @@").ToString(), Pin, Node, BetterSelf);
 						}
 					}
-				};
-				FStructConnectHelper::Connect(DefaultToSelfMetaValue, this, BetterSelfPin, Schema, CompilerContext.MessageLog);
-				FStructConnectHelper::Connect(WorldContextMetaValue, this, BetterSelfPin, Schema, CompilerContext.MessageLog);
-			}
+				}
+			};
+			FStructConnectHelper::Connect(DefaultToSelfMetaValue, this, BetterSelfPin, Schema, CompilerContext.MessageLog);
+			FStructConnectHelper::Connect(WorldContextMetaValue, this, BetterSelfPin, Schema, CompilerContext.MessageLog);
 		}
+	}
 
-		// If we have an enum param that is expanded, we handle that first
-		if(bWantsEnumToExecExpansion)
+	// If we have an enum param that is expanded, we handle that first
+	if(bWantsEnumToExecExpansion)
+	{
+		if(Function)
 		{
-			if(Function)
+			// Get the metadata that identifies which param is the enum, and try and find it
+			const FString& EnumParamName = Function->GetMetaData(FBlueprintMetadata::MD_ExpandEnumAsExecs);
+			UByteProperty* EnumProp = FindField<UByteProperty>(Function, FName(*EnumParamName));
+			UEdGraphPin* EnumParamPin = FindPinChecked(EnumParamName);
+			if(EnumProp != NULL && EnumProp->Enum != NULL)
 			{
-				// Get the metadata that identifies which param is the enum, and try and find it
-				const FString& EnumParamName = Function->GetMetaData(FBlueprintMetadata::MD_ExpandEnumAsExecs);
-				UByteProperty* EnumProp = FindField<UByteProperty>(Function, FName(*EnumParamName));
-				UEdGraphPin* EnumParamPin = FindPinChecked(EnumParamName);
-				if(EnumProp != NULL && EnumProp->Enum != NULL)
+				// Expanded as input execs pins
+				if (EnumParamPin->Direction == EGPD_Input)
 				{
-					// Expanded as input execs pins
-					if (EnumParamPin->Direction == EGPD_Input)
+					// Create normal exec input
+					UEdGraphPin* ExecutePin = CreatePin(EGPD_Input, Schema->PC_Exec, TEXT(""), NULL, false, false, Schema->PN_Execute);
+
+					// Create temp enum variable
+					UK2Node_TemporaryVariable* TempEnumVarNode = CompilerContext.SpawnIntermediateNode<UK2Node_TemporaryVariable>(this, SourceGraph);
+					TempEnumVarNode->VariableType.PinCategory = Schema->PC_Byte;
+					TempEnumVarNode->VariableType.PinSubCategoryObject = EnumProp->Enum;
+					TempEnumVarNode->AllocateDefaultPins();
+					// Get the output pin
+					UEdGraphPin* TempEnumVarOutput = TempEnumVarNode->GetVariablePin();
+
+					// Connect temp enum variable to (hidden) enum pin
+					Schema->TryCreateConnection(TempEnumVarOutput, EnumParamPin);
+
+					// Now we want to iterate over other exec inputs...
+					for(int32 PinIdx=Pins.Num()-1; PinIdx>=0; PinIdx--)
 					{
-						// Create normal exec input
-						UEdGraphPin* ExecutePin = CreatePin(EGPD_Input, Schema->PC_Exec, TEXT(""), NULL, false, false, Schema->PN_Execute);
-
-						// Create temp enum variable
-						UK2Node_TemporaryVariable* TempEnumVarNode = CompilerContext.SpawnIntermediateNode<UK2Node_TemporaryVariable>(this, SourceGraph);
-						TempEnumVarNode->VariableType.PinCategory = Schema->PC_Byte;
-						TempEnumVarNode->VariableType.PinSubCategoryObject = EnumProp->Enum;
-						TempEnumVarNode->AllocateDefaultPins();
-						// Get the output pin
-						UEdGraphPin* TempEnumVarOutput = TempEnumVarNode->GetVariablePin();
-
-						// Connect temp enum variable to (hidden) enum pin
-						Schema->TryCreateConnection(TempEnumVarOutput, EnumParamPin);
-
-						// Now we want to iterate over other exec inputs...
-						for(int32 PinIdx=Pins.Num()-1; PinIdx>=0; PinIdx--)
+						UEdGraphPin* Pin = Pins[PinIdx];
+						if( Pin != NULL && 
+							Pin != ExecutePin &&
+							Pin->Direction == EGPD_Input && 
+							Pin->PinType.PinCategory == Schema->PC_Exec )
 						{
-							UEdGraphPin* Pin = Pins[PinIdx];
-							if( Pin != NULL && 
-								Pin != ExecutePin &&
-								Pin->Direction == EGPD_Input && 
-								Pin->PinType.PinCategory == Schema->PC_Exec )
-							{
-								// Create node to set the temp enum var
-								UK2Node_AssignmentStatement* AssignNode = CompilerContext.SpawnIntermediateNode<UK2Node_AssignmentStatement>(this, SourceGraph);
-								AssignNode->AllocateDefaultPins();
-
-								// Move connections from fake 'enum exec' pint to this assignment node
-									CompilerContext.MovePinLinksToIntermediate(*Pin, *AssignNode->GetExecPin());
-
-								// Connect this to out temp enum var
-								Schema->TryCreateConnection(AssignNode->GetVariablePin(), TempEnumVarOutput);
-
-								// Connect exec output to 'real' exec pin
-								Schema->TryCreateConnection(AssignNode->GetThenPin(), ExecutePin);
-
-								// set the literal enum value to set to
-								AssignNode->GetValuePin()->DefaultValue = Pin->PinName;
-
-								// Finally remove this 'cosmetic' exec pin
-								Pins.RemoveAt(PinIdx);
-							}
+							// Create node to set the temp enum var
+							UK2Node_AssignmentStatement* AssignNode = CompilerContext.SpawnIntermediateNode<UK2Node_AssignmentStatement>(this, SourceGraph);
+							AssignNode->AllocateDefaultPins();
+
+							// Move connections from fake 'enum exec' pint to this assignment node
+								CompilerContext.MovePinLinksToIntermediate(*Pin, *AssignNode->GetExecPin());
+
+							// Connect this to out temp enum var
+							Schema->TryCreateConnection(AssignNode->GetVariablePin(), TempEnumVarOutput);
+
+							// Connect exec output to 'real' exec pin
+							Schema->TryCreateConnection(AssignNode->GetThenPin(), ExecutePin);
+
+							// set the literal enum value to set to
+							AssignNode->GetValuePin()->DefaultValue = Pin->PinName;
+
+							// Finally remove this 'cosmetic' exec pin
+							Pins.RemoveAt(PinIdx);
 						}
 					}
-					// Expanded as output execs pins
-					else if (EnumParamPin->Direction == EGPD_Output)
-					{
-						// Create normal exec output
-						UEdGraphPin* ExecutePin = CreatePin(EGPD_Output, Schema->PC_Exec, TEXT(""), NULL, false, false, Schema->PN_Execute);
+				}
+				// Expanded as output execs pins
+				else if (EnumParamPin->Direction == EGPD_Output)
+				{
+					// Create normal exec output
+					UEdGraphPin* ExecutePin = CreatePin(EGPD_Output, Schema->PC_Exec, TEXT(""), NULL, false, false, Schema->PN_Execute);
 						
-						// Create a SwitchEnum node to switch on the output enum
-						UK2Node_SwitchEnum* SwitchEnumNode = CompilerContext.SpawnIntermediateNode<UK2Node_SwitchEnum>(this, SourceGraph);
-						UEnum* EnumObject = Cast<UEnum>(EnumParamPin->PinType.PinSubCategoryObject.Get());
-						SwitchEnumNode->SetEnum(EnumObject);
-						SwitchEnumNode->AllocateDefaultPins();
+					// Create a SwitchEnum node to switch on the output enum
+					UK2Node_SwitchEnum* SwitchEnumNode = CompilerContext.SpawnIntermediateNode<UK2Node_SwitchEnum>(this, SourceGraph);
+					UEnum* EnumObject = Cast<UEnum>(EnumParamPin->PinType.PinSubCategoryObject.Get());
+					SwitchEnumNode->SetEnum(EnumObject);
+					SwitchEnumNode->AllocateDefaultPins();
 						
-						// Hook up execution to the switch node
-						Schema->TryCreateConnection(ExecutePin, SwitchEnumNode->GetExecPin());
-						// Connect (hidden) enum pin to switch node's selection pin
-						Schema->TryCreateConnection(EnumParamPin, SwitchEnumNode->GetSelectionPin());
+					// Hook up execution to the switch node
+					Schema->TryCreateConnection(ExecutePin, SwitchEnumNode->GetExecPin());
+					// Connect (hidden) enum pin to switch node's selection pin
+					Schema->TryCreateConnection(EnumParamPin, SwitchEnumNode->GetSelectionPin());
 						
-						// Now we want to iterate over other exec outputs
-						for(int32 PinIdx=Pins.Num()-1; PinIdx>=0; PinIdx--)
+					// Now we want to iterate over other exec outputs
+					for(int32 PinIdx=Pins.Num()-1; PinIdx>=0; PinIdx--)
+					{
+						UEdGraphPin* Pin = Pins[PinIdx];
+						if( Pin != NULL &&
+							Pin != ExecutePin &&
+							Pin->Direction == EGPD_Output &&
+							Pin->PinType.PinCategory == Schema->PC_Exec )
 						{
-							UEdGraphPin* Pin = Pins[PinIdx];
-							if( Pin != NULL &&
-							   Pin != ExecutePin &&
-							   Pin->Direction == EGPD_Output &&
-							   Pin->PinType.PinCategory == Schema->PC_Exec )
-							{
-								// Move connections from fake 'enum exec' pint to this switch node
-								CompilerContext.MovePinLinksToIntermediate(*Pin, *SwitchEnumNode->FindPinChecked(Pin->PinName));
+							// Move connections from fake 'enum exec' pint to this switch node
+							CompilerContext.MovePinLinksToIntermediate(*Pin, *SwitchEnumNode->FindPinChecked(Pin->PinName));
 								
-								// Finally remove this 'cosmetic' exec pin
-								Pins.RemoveAt(PinIdx);
-							}
+							// Finally remove this 'cosmetic' exec pin
+							Pins.RemoveAt(PinIdx);
 						}
 					}
 				}
 			}
 		}
+	}
 
-		// AUTO CREATED REFS
+	// AUTO CREATED REFS
+	{
+		if ( Function )
 		{
-			if ( Function )
+			TArray<FString> AutoCreateRefTermPinNames;
+			const bool bHasAutoCreateRefTerms = Function->HasMetaData(FBlueprintMetadata::MD_AutoCreateRefTerm);
+			if ( bHasAutoCreateRefTerms )
 			{
-				TArray<FString> AutoCreateRefTermPinNames;
-				const bool bHasAutoCreateRefTerms = Function->HasMetaData(FBlueprintMetadata::MD_AutoCreateRefTerm);
-				if ( bHasAutoCreateRefTerms )
-				{
-					CompilerContext.GetSchema()->GetAutoEmitTermParameters(Function, AutoCreateRefTermPinNames);
-				}
+				CompilerContext.GetSchema()->GetAutoEmitTermParameters(Function, AutoCreateRefTermPinNames);
+			}
 
-				for ( auto Pin : Pins )
+			for ( auto Pin : Pins )
+			{
+				if ( Pin && bHasAutoCreateRefTerms && AutoCreateRefTermPinNames.Contains(Pin->PinName) )
 				{
-					if ( Pin && bHasAutoCreateRefTerms && AutoCreateRefTermPinNames.Contains(Pin->PinName) )
+					const bool bHasDefaultValue = !Pin->DefaultValue.IsEmpty() || Pin->DefaultObject || !Pin->DefaultTextValue.IsEmpty();
+					const bool bValidAutoRefPin = Pin->PinType.bIsReference
+						&& !CompilerContext.GetSchema()->IsMetaPin(*Pin)
+						&& ( Pin->Direction == EGPD_Input )
+						&& !Pin->LinkedTo.Num()
+						&& ( Pin->PinType.bIsArray || bHasDefaultValue );
+					if ( bValidAutoRefPin )
 					{
-						const bool bHasDefaultValue = !Pin->DefaultValue.IsEmpty() || Pin->DefaultObject || !Pin->DefaultTextValue.IsEmpty();
-						const bool bValidAutoRefPin = Pin->PinType.bIsReference
-							&& !CompilerContext.GetSchema()->IsMetaPin(*Pin)
-							&& ( Pin->Direction == EGPD_Input )
-							&& !Pin->LinkedTo.Num()
-							&& ( Pin->PinType.bIsArray || bHasDefaultValue );
-						if ( bValidAutoRefPin )
+						//default values can be reset when the pin is connected
+						const auto DefaultValue = Pin->DefaultValue;
+						const auto DefaultObject = Pin->DefaultObject;
+						const auto DefaultTextValue = Pin->DefaultTextValue;
+
+						auto ValuePin = InnerHandleAutoCreateRef(this, Pin, CompilerContext, SourceGraph, bHasDefaultValue);
+						if ( ValuePin )
 						{
-							//default values can be reset when the pin is connected
-							const auto DefaultValue = Pin->DefaultValue;
-							const auto DefaultObject = Pin->DefaultObject;
-							const auto DefaultTextValue = Pin->DefaultTextValue;
-
-							auto ValuePin = InnerHandleAutoCreateRef(this, Pin, CompilerContext, SourceGraph, bHasDefaultValue);
-							if ( ValuePin )
-							{
-								ValuePin->DefaultValue = DefaultValue;
-								ValuePin->DefaultObject = DefaultObject;
-								ValuePin->DefaultTextValue = DefaultTextValue;
-							}
+							ValuePin->DefaultValue = DefaultValue;
+							ValuePin->DefaultObject = DefaultObject;
+							ValuePin->DefaultTextValue = DefaultTextValue;
 						}
 					}
 				}
 			}
 		}
+	}
 
-		// Then we go through and expand out array iteration if necessary
-		const bool bAllowMultipleSelfs = AllowMultipleSelfs(true);
-		UEdGraphPin* MultiSelf = Schema->FindSelfPin(*this, EEdGraphPinDirection::EGPD_Input);
-		if(bAllowMultipleSelfs && MultiSelf && !MultiSelf->PinType.bIsArray)
+	// Then we go through and expand out array iteration if necessary
+	const bool bAllowMultipleSelfs = AllowMultipleSelfs(true);
+	UEdGraphPin* MultiSelf = Schema->FindSelfPin(*this, EEdGraphPinDirection::EGPD_Input);
+	if(bAllowMultipleSelfs && MultiSelf && !MultiSelf->PinType.bIsArray)
+	{
+		const bool bProperInputToExpandForEach = 
+			(1 == MultiSelf->LinkedTo.Num()) && 
+			(NULL != MultiSelf->LinkedTo[0]) && 
+			(MultiSelf->LinkedTo[0]->PinType.bIsArray);
+		if(bProperInputToExpandForEach)
 		{
-			const bool bProperInputToExpandForEach = 
-				(1 == MultiSelf->LinkedTo.Num()) && 
-				(NULL != MultiSelf->LinkedTo[0]) && 
-				(MultiSelf->LinkedTo[0]->PinType.bIsArray);
-			if(bProperInputToExpandForEach)
-			{
-				CallForEachElementInArrayExpansion(this, MultiSelf, CompilerContext, SourceGraph);
-			}
+			CallForEachElementInArrayExpansion(this, MultiSelf, CompilerContext, SourceGraph);
 		}
 	}
 }
diff --git a/Engine/Source/Editor/BlueprintGraph/Private/K2Node_CallFunctionOnMember.cpp b/Engine/Source/Editor/BlueprintGraph/Private/K2Node_CallFunctionOnMember.cpp
--- a/Engine/Source/Editor/BlueprintGraph/Private/K2Node_CallFunctionOnMember.cpp
+++ b/Engine/Source/Editor/BlueprintGraph/Private/K2Node_CallFunctionOnMember.cpp
@@ -50,79 +50,76 @@ void UK2Node_CallFunctionOnMember::ExpandNode(class FKismetCompilerContext& Comp
 	// and then that CallFunction node will get its own Expansion to handle the parent portions
 	UK2Node::ExpandNode(CompilerContext, SourceGraph);
 
-	if (CompilerContext.bIsFullCompile)
+	const UEdGraphSchema_K2* Schema = CompilerContext.GetSchema();
+	UFunction* Function = GetTargetFunction();
+
+	// Create real 'call function' node.
+	UK2Node_CallFunction* CallFuncNode = CompilerContext.SpawnIntermediateNode<UK2Node_CallFunction>(this, SourceGraph);
+	CallFuncNode->SetFromFunction(Function);
+	CallFuncNode->AllocateDefaultPins();
+	UEdGraphPin* CallFuncSelfPin = Schema->FindSelfPin(*CallFuncNode, EGPD_Input);
+
+	// Now because you can wire multiple variables to a self pin, need to iterate over each one and create a 'get var' node for each
+	UEdGraphPin* SelfPin = Schema->FindSelfPin(*this, EGPD_Input);
+	if(SelfPin != NULL)
 	{
-		const UEdGraphSchema_K2* Schema = CompilerContext.GetSchema();
-		UFunction* Function = GetTargetFunction();
-
-		// Create real 'call function' node.
-		UK2Node_CallFunction* CallFuncNode = CompilerContext.SpawnIntermediateNode<UK2Node_CallFunction>(this, SourceGraph);
-		CallFuncNode->SetFromFunction(Function);
-		CallFuncNode->AllocateDefaultPins();
-		UEdGraphPin* CallFuncSelfPin = Schema->FindSelfPin(*CallFuncNode, EGPD_Input);
-
-		// Now because you can wire multiple variables to a self pin, need to iterate over each one and create a 'get var' node for each
-		UEdGraphPin* SelfPin = Schema->FindSelfPin(*this, EGPD_Input);
-		if(SelfPin != NULL)
+		if (SelfPin->LinkedTo.Num() == 0)
 		{
-			if (SelfPin->LinkedTo.Num() == 0)
-			{
-				UK2Node_VariableGet* GetVarNode = CompilerContext.SpawnIntermediateNode<UK2Node_VariableGet>(this, SourceGraph);
-				GetVarNode->VariableReference.SetSelfMember(MemberVariableToCallOn.GetMemberName());
-				GetVarNode->AllocateDefaultPins();
+			UK2Node_VariableGet* GetVarNode = CompilerContext.SpawnIntermediateNode<UK2Node_VariableGet>(this, SourceGraph);
+			GetVarNode->VariableReference.SetSelfMember(MemberVariableToCallOn.GetMemberName());
+			GetVarNode->AllocateDefaultPins();
 
-				if (UEdGraphPin* ValuePin = GetVarNode->GetValuePin())
-				{
-					ValuePin->MakeLinkTo(CallFuncSelfPin);
-				}
+			if (UEdGraphPin* ValuePin = GetVarNode->GetValuePin())
+			{
+				ValuePin->MakeLinkTo(CallFuncSelfPin);
 			}
-			else
+		}
+		else
+		{
+			for (int32 TargetIdx = 0; TargetIdx < SelfPin->LinkedTo.Num(); TargetIdx++)
 			{
-				for (int32 TargetIdx = 0; TargetIdx < SelfPin->LinkedTo.Num(); TargetIdx++)
+				UEdGraphPin* SourcePin = SelfPin->LinkedTo[TargetIdx];
+				if (SourcePin != NULL)
 				{
-					UEdGraphPin* SourcePin = SelfPin->LinkedTo[TargetIdx];
-					if (SourcePin != NULL)
+					// Create 'get var' node to get the member
+					UK2Node_VariableGet* GetVarNode = CompilerContext.SpawnIntermediateNode<UK2Node_VariableGet>(this, SourceGraph);
+					GetVarNode->VariableReference = MemberVariableToCallOn;
+					GetVarNode->AllocateDefaultPins();
+
+					UEdGraphPin* VarNodeSelfPin = Schema->FindSelfPin(*GetVarNode, EGPD_Input);
+					if (VarNodeSelfPin != NULL)
+					{
+						VarNodeSelfPin->MakeLinkTo(SourcePin);
+
+						UEdGraphPin* ValuePin = GetVarNode->GetValuePin();
+						ValuePin->MakeLinkTo(CallFuncSelfPin);
+					}
+					else
 					{
-						// Create 'get var' node to get the member
-						UK2Node_VariableGet* GetVarNode = CompilerContext.SpawnIntermediateNode<UK2Node_VariableGet>(this, SourceGraph);
-						GetVarNode->VariableReference = MemberVariableToCallOn;
-						GetVarNode->AllocateDefaultPins();
-
-						UEdGraphPin* VarNodeSelfPin = Schema->FindSelfPin(*GetVarNode, EGPD_Input);
-						if (VarNodeSelfPin != NULL)
-						{
-							VarNodeSelfPin->MakeLinkTo(SourcePin);
-
-							UEdGraphPin* ValuePin = GetVarNode->GetValuePin();
-							ValuePin->MakeLinkTo(CallFuncSelfPin);
-						}
-						else
-						{
-							// Failed to find the member to call on for this expansion, so warn about it
-							CompilerContext.MessageLog.Warning(*LOCTEXT("CallFunctionOnInvalidMember_Warning", "Function node @@ called on invalid target member.").ToString(), this);
-						}
+						// Failed to find the member to call on for this expansion, so warn about it
+						CompilerContext.MessageLog.Warning(*LOCTEXT("CallFunctionOnInvalidMember_Warning", "Function node @@ called on invalid target member.").ToString(), this);
 					}
 				}
-			}			
-		}
+			}
+		}			
+	}
 
-		// Now move the rest of the connections (including exec connections...)
-		for(int32 SrcPinIdx=0; SrcPinIdx<Pins.Num(); SrcPinIdx++)
+	// Now move the rest of the connections (including exec connections...)
+	for(int32 SrcPinIdx=0; SrcPinIdx<Pins.Num(); SrcPinIdx++)
+	{
+		UEdGraphPin* SrcPin = Pins[SrcPinIdx];
+		if(SrcPin != NULL && SrcPin != SelfPin) // check its not the self pin
 		{
-			UEdGraphPin* SrcPin = Pins[SrcPinIdx];
-			if(SrcPin != NULL && SrcPin != SelfPin) // check its not the self pin
+			UEdGraphPin* DestPin = CallFuncNode->FindPin(SrcPin->PinName);
+			if(DestPin != NULL)
 			{
-				UEdGraphPin* DestPin = CallFuncNode->FindPin(SrcPin->PinName);
-				if(DestPin != NULL)
-				{
-					CompilerContext.MovePinLinksToIntermediate(*SrcPin, *DestPin); // Source node is assumed to be owner...
-				}
+				CompilerContext.MovePinLinksToIntermediate(*SrcPin, *DestPin); // Source node is assumed to be owner...
 			}
 		}
-
-		// Finally, break any remaining links on the 'call func on member' node
-		BreakAllNodeLinks();
 	}
+
+	// Finally, break any remaining links on the 'call func on member' node
+	BreakAllNodeLinks();
 }
 
 bool UK2Node_CallFunctionOnMember::HasExternalBlueprintDependencies(TArray<class UStruct*>* OptionalOutput) const
diff --git a/Engine/Source/Editor/BlueprintGraph/Private/K2Node_CastByteToEnum.cpp b/Engine/Source/Editor/BlueprintGraph/Private/K2Node_CastByteToEnum.cpp
--- a/Engine/Source/Editor/BlueprintGraph/Private/K2Node_CastByteToEnum.cpp
+++ b/Engine/Source/Editor/BlueprintGraph/Private/K2Node_CastByteToEnum.cpp
@@ -59,7 +59,7 @@ void UK2Node_CastByteToEnum::ExpandNode(class FKismetCompilerContext& CompilerCo
 {
 	Super::ExpandNode(CompilerContext, SourceGraph);
 
-	if (CompilerContext.bIsFullCompile && bSafe)
+	if (bSafe)
 	{
 		const UEdGraphSchema_K2* Schema = CompilerContext.GetSchema();
 
diff --git a/Engine/Source/Editor/BlueprintGraph/Private/K2Node_CommutativeAssociativeBinaryOperator.cpp b/Engine/Source/Editor/BlueprintGraph/Private/K2Node_CommutativeAssociativeBinaryOperator.cpp
--- a/Engine/Source/Editor/BlueprintGraph/Private/K2Node_CommutativeAssociativeBinaryOperator.cpp
+++ b/Engine/Source/Editor/BlueprintGraph/Private/K2Node_CommutativeAssociativeBinaryOperator.cpp
@@ -253,7 +253,7 @@ void UK2Node_CommutativeAssociativeBinaryOperator::ExpandNode(FKismetCompilerCon
 {
 	Super::ExpandNode(CompilerContext, SourceGraph);
 
-	if (CompilerContext.bIsFullCompile && NumAdditionalInputs > 0)
+	if (NumAdditionalInputs > 0)
 	{
 		const UEdGraphSchema_K2* Schema = CompilerContext.GetSchema();
 
diff --git a/Engine/Source/Editor/BlueprintGraph/Private/K2Node_DelegateSet.cpp b/Engine/Source/Editor/BlueprintGraph/Private/K2Node_DelegateSet.cpp
--- a/Engine/Source/Editor/BlueprintGraph/Private/K2Node_DelegateSet.cpp
+++ b/Engine/Source/Editor/BlueprintGraph/Private/K2Node_DelegateSet.cpp
@@ -263,7 +263,7 @@ void UK2Node_DelegateSet::ExpandNode(class FKismetCompilerContext& CompilerConte
 	{
 		CompilerContext.MessageLog.Error(*FString::Printf(*NSLOCTEXT("KismetCompiler", "InvalidNodeOutsideUbergraph_Error", "Unexpected node @@ found outside ubergraph.").ToString()), this);
 	}
-	else if (CompilerContext.bIsFullCompile)
+	else
 	{
 		UFunction* TargetFunction = GetDelegateSignature();
 		if(TargetFunction != NULL)
diff --git a/Engine/Source/Editor/BlueprintGraph/Private/K2Node_DoOnceMultiInput.cpp b/Engine/Source/Editor/BlueprintGraph/Private/K2Node_DoOnceMultiInput.cpp
--- a/Engine/Source/Editor/BlueprintGraph/Private/K2Node_DoOnceMultiInput.cpp
+++ b/Engine/Source/Editor/BlueprintGraph/Private/K2Node_DoOnceMultiInput.cpp
@@ -327,85 +327,82 @@ void UK2Node_DoOnceMultiInput::ExpandNode(FKismetCompilerContext& CompilerContex
 {
 	Super::ExpandNode(CompilerContext, SourceGraph);
 
-	if (CompilerContext.bIsFullCompile)
-	{
-		const UEdGraphSchema_K2* Schema = CompilerContext.GetSchema();
+	const UEdGraphSchema_K2* Schema = CompilerContext.GetSchema();
 
-		/////////////////////////////
-		// Temporary Variable node
-		/////////////////////////////
+	/////////////////////////////
+	// Temporary Variable node
+	/////////////////////////////
 
-		// Create the node
-		UK2Node_TemporaryVariable* TempVarNode = SourceGraph->CreateBlankNode<UK2Node_TemporaryVariable>();
-		TempVarNode->VariableType.PinCategory = Schema->PC_Boolean;
-		TempVarNode->AllocateDefaultPins();
-		CompilerContext.MessageLog.NotifyIntermediateObjectCreation(TempVarNode, this);
-		// Give a reference of the variable node to the multi gate node
-		DataNode = TempVarNode;
+	// Create the node
+	UK2Node_TemporaryVariable* TempVarNode = SourceGraph->CreateBlankNode<UK2Node_TemporaryVariable>();
+	TempVarNode->VariableType.PinCategory = Schema->PC_Boolean;
+	TempVarNode->AllocateDefaultPins();
+	CompilerContext.MessageLog.NotifyIntermediateObjectCreation(TempVarNode, this);
+	// Give a reference of the variable node to the multi gate node
+	DataNode = TempVarNode;
 
-		// Create the conditional node we're replacing the enum node for
+	// Create the conditional node we're replacing the enum node for
 
-		for (int32 idx = 0; idx < NumBaseInputs + NumAdditionalInputs + 1 /*ResetPin*/; ++idx)
-		{
-			UEdGraphPin* ExecPin = GetInputPin(idx);
-			UEdGraphPin* ThenPin = GetOutputPin(idx);
-
-			check(ExecPin);
-			check(ThenPin);
-
-			// AssignmentNode
-			UK2Node_AssignmentStatement* AssignmentNode = SourceGraph->CreateBlankNode<UK2Node_AssignmentStatement>();
-			AssignmentNode->AllocateDefaultPins();
-			CompilerContext.MessageLog.NotifyIntermediateObjectCreation(AssignmentNode, this);
-			AssignmentNode->GetVariablePin()->PinType = TempVarNode->GetVariablePin()->PinType;
-			AssignmentNode->GetVariablePin()->MakeLinkTo(TempVarNode->GetVariablePin());
-			AssignmentNode->GetValuePin()->PinType = TempVarNode->GetVariablePin()->PinType;
+	for (int32 idx = 0; idx < NumBaseInputs + NumAdditionalInputs + 1 /*ResetPin*/; ++idx)
+	{
+		UEdGraphPin* ExecPin = GetInputPin(idx);
+		UEdGraphPin* ThenPin = GetOutputPin(idx);
+
+		check(ExecPin);
+		check(ThenPin);
+
+		// AssignmentNode
+		UK2Node_AssignmentStatement* AssignmentNode = SourceGraph->CreateBlankNode<UK2Node_AssignmentStatement>();
+		AssignmentNode->AllocateDefaultPins();
+		CompilerContext.MessageLog.NotifyIntermediateObjectCreation(AssignmentNode, this);
+		AssignmentNode->GetVariablePin()->PinType = TempVarNode->GetVariablePin()->PinType;
+		AssignmentNode->GetVariablePin()->MakeLinkTo(TempVarNode->GetVariablePin());
+		AssignmentNode->GetValuePin()->PinType = TempVarNode->GetVariablePin()->PinType;
 			
-			if (!ExecPin->PinName.Contains(TEXT("Reset"))) // Fixme this wont work for localization
-			{
-				// BranchNode
-				UK2Node_IfThenElse* BranchNode = SourceGraph->CreateBlankNode<UK2Node_IfThenElse>();
-				BranchNode->AllocateDefaultPins();
-				CompilerContext.MessageLog.NotifyIntermediateObjectCreation(BranchNode, this);
+		if (!ExecPin->PinName.Contains(TEXT("Reset"))) // Fixme this wont work for localization
+		{
+			// BranchNode
+			UK2Node_IfThenElse* BranchNode = SourceGraph->CreateBlankNode<UK2Node_IfThenElse>();
+			BranchNode->AllocateDefaultPins();
+			CompilerContext.MessageLog.NotifyIntermediateObjectCreation(BranchNode, this);
 
-				// -------------------------------------------------
+			// -------------------------------------------------
 
-				// Coerce the wildcards pin types (set the default of the value to 0)
-				AssignmentNode->GetValuePin()->DefaultValue = TEXT("1");
+			// Coerce the wildcards pin types (set the default of the value to 0)
+			AssignmentNode->GetValuePin()->DefaultValue = TEXT("1");
 
-				// -------------------------------------------------
+			// -------------------------------------------------
 
-				// Link Tempvariable with the branch condtional
-				Schema->TryCreateConnection(TempVarNode->GetVariablePin(), BranchNode->GetConditionPin());
+			// Link Tempvariable with the branch condtional
+			Schema->TryCreateConnection(TempVarNode->GetVariablePin(), BranchNode->GetConditionPin());
 			
-				// Link our input exec pin into the branch node
-				CompilerContext.MovePinLinksToIntermediate(*ExecPin, *BranchNode->GetExecPin());
+			// Link our input exec pin into the branch node
+			CompilerContext.MovePinLinksToIntermediate(*ExecPin, *BranchNode->GetExecPin());
 
-				// link branch else (false) to assigment node (set temp variable to true)
-				Schema->TryCreateConnection(BranchNode->GetElsePin(), AssignmentNode->GetExecPin());
+			// link branch else (false) to assigment node (set temp variable to true)
+			Schema->TryCreateConnection(BranchNode->GetElsePin(), AssignmentNode->GetExecPin());
 
-				// link set temp variable node to our ouput then pin
-				CompilerContext.MovePinLinksToIntermediate(*ThenPin, *AssignmentNode->GetThenPin());
+			// link set temp variable node to our ouput then pin
+			CompilerContext.MovePinLinksToIntermediate(*ThenPin, *AssignmentNode->GetThenPin());
 
-			}
-			else
-			{
-				// Coerce the wildcards pin types (set the default of the value to 1)				
-				AssignmentNode->GetValuePin()->DefaultValue = TEXT("0");
+		}
+		else
+		{
+			// Coerce the wildcards pin types (set the default of the value to 1)				
+			AssignmentNode->GetValuePin()->DefaultValue = TEXT("0");
 
-				// -------------------------------------------------
+			// -------------------------------------------------
 
-				// Link our input exec pin into the branch node
-				CompilerContext.MovePinLinksToIntermediate(*ExecPin, *AssignmentNode->GetExecPin());
+			// Link our input exec pin into the branch node
+			CompilerContext.MovePinLinksToIntermediate(*ExecPin, *AssignmentNode->GetExecPin());
 
-				// link set temp variable node to our ouput then pin
-				CompilerContext.MovePinLinksToIntermediate(*ThenPin, *AssignmentNode->GetThenPin());
-			}
+			// link set temp variable node to our ouput then pin
+			CompilerContext.MovePinLinksToIntermediate(*ThenPin, *AssignmentNode->GetThenPin());
 		}
-
-		// Break all links to the Select node so it goes away for at scheduling time
-		BreakAllNodeLinks();
 	}
+
+	// Break all links to the Select node so it goes away for at scheduling time
+	BreakAllNodeLinks();
 }
 
 void UK2Node_DoOnceMultiInput::GetMenuActions(FBlueprintActionDatabaseRegistrar& ActionRegistrar) const
diff --git a/Engine/Source/Editor/BlueprintGraph/Private/K2Node_EaseFunction.cpp b/Engine/Source/Editor/BlueprintGraph/Private/K2Node_EaseFunction.cpp
--- a/Engine/Source/Editor/BlueprintGraph/Private/K2Node_EaseFunction.cpp
+++ b/Engine/Source/Editor/BlueprintGraph/Private/K2Node_EaseFunction.cpp
@@ -384,48 +384,45 @@ void UK2Node_EaseFunction::ExpandNode(class FKismetCompilerContext& CompilerCont
 {
 	Super::ExpandNode(CompilerContext, SourceGraph);
 
-	if (CompilerContext.bIsFullCompile)
-	{
-		/**
+	/**
 		At the end of this, the UK2Node_EaseFunction will not be a part of the Blueprint, it merely handles connecting
 		the other nodes into the Blueprint.
-		*/
+	*/
 
-		UFunction* function = UKismetMathLibrary::StaticClass()->FindFunctionByName(*EaseFunctionName);
-		if (function == NULL)
-		{
-			CompilerContext.MessageLog.Error(*LOCTEXT("InvalidFunctionName", "BaseAsyncTask: Type not supported or not initialized. @@").ToString(), this);
-			return;
-		}
+	UFunction* function = UKismetMathLibrary::StaticClass()->FindFunctionByName(*EaseFunctionName);
+	if (function == NULL)
+	{
+		CompilerContext.MessageLog.Error(*LOCTEXT("InvalidFunctionName", "BaseAsyncTask: Type not supported or not initialized. @@").ToString(), this);
+		return;
+	}
 
-		const UEdGraphSchema_K2* Schema = CompilerContext.GetSchema();
+	const UEdGraphSchema_K2* Schema = CompilerContext.GetSchema();
 
-		// The call function does all the real work, each child class implementing easing for  a given type provides
-		// the name of the desired function
-		UK2Node_CallFunction* CallFunction = CompilerContext.SpawnIntermediateNode<UK2Node_CallFunction>(this, SourceGraph);
+	// The call function does all the real work, each child class implementing easing for  a given type provides
+	// the name of the desired function
+	UK2Node_CallFunction* CallFunction = CompilerContext.SpawnIntermediateNode<UK2Node_CallFunction>(this, SourceGraph);
 		
-		CallFunction->SetFromFunction(function);
-		CallFunction->AllocateDefaultPins();
-		CompilerContext.MessageLog.NotifyIntermediateObjectCreation(CallFunction, this);
+	CallFunction->SetFromFunction(function);
+	CallFunction->AllocateDefaultPins();
+	CompilerContext.MessageLog.NotifyIntermediateObjectCreation(CallFunction, this);
 
-		// Move the ease function and the alpha connections from us to the call function
-		CompilerContext.MovePinLinksToIntermediate(*FindPin(FEaseFunctionNodeHelper::GetEaseFuncPinName()), *CallFunction->FindPin(TEXT("EasingFunc")));
-		CompilerContext.MovePinLinksToIntermediate(*FindPin(FEaseFunctionNodeHelper::GetAlphaPinName()), *CallFunction->FindPin(TEXT("Alpha")));
+	// Move the ease function and the alpha connections from us to the call function
+	CompilerContext.MovePinLinksToIntermediate(*FindPin(FEaseFunctionNodeHelper::GetEaseFuncPinName()), *CallFunction->FindPin(TEXT("EasingFunc")));
+	CompilerContext.MovePinLinksToIntermediate(*FindPin(FEaseFunctionNodeHelper::GetAlphaPinName()), *CallFunction->FindPin(TEXT("Alpha")));
 
-		// Move base connections to the call function's connections
-		CompilerContext.MovePinLinksToIntermediate(*FindPin(FEaseFunctionNodeHelper::GetAPinName()), *CallFunction->FindPin(TEXT("A")));
-		CompilerContext.MovePinLinksToIntermediate(*FindPin(FEaseFunctionNodeHelper::GetBPinName()), *CallFunction->FindPin(TEXT("B")));
-		CompilerContext.MovePinLinksToIntermediate(*FindPin(FEaseFunctionNodeHelper::GetResultPinName()), *CallFunction->GetReturnValuePin());
+	// Move base connections to the call function's connections
+	CompilerContext.MovePinLinksToIntermediate(*FindPin(FEaseFunctionNodeHelper::GetAPinName()), *CallFunction->FindPin(TEXT("A")));
+	CompilerContext.MovePinLinksToIntermediate(*FindPin(FEaseFunctionNodeHelper::GetBPinName()), *CallFunction->FindPin(TEXT("B")));
+	CompilerContext.MovePinLinksToIntermediate(*FindPin(FEaseFunctionNodeHelper::GetResultPinName()), *CallFunction->GetReturnValuePin());
 
-		// Now move the custom pins to their new locations
-		for (int32 ArgIdx = 0; ArgIdx < CustomPinNames.Num(); ++ArgIdx)
-		{
-			CompilerContext.MovePinLinksToIntermediate(*FindPin(CustomPinNames[ArgIdx].PinName), *CallFunction->FindPin(CustomPinNames[ArgIdx].CallFuncPinName));
-		}
-
-		// Cleanup links to ourselfs and we are done!
-		BreakAllNodeLinks();
+	// Now move the custom pins to their new locations
+	for (int32 ArgIdx = 0; ArgIdx < CustomPinNames.Num(); ++ArgIdx)
+	{
+		CompilerContext.MovePinLinksToIntermediate(*FindPin(CustomPinNames[ArgIdx].PinName), *CallFunction->FindPin(CustomPinNames[ArgIdx].CallFuncPinName));
 	}
+
+	// Cleanup links to ourselfs and we are done!
+	BreakAllNodeLinks();
 }
 
 UEdGraphPin* UK2Node_EaseFunction::GetEaseFuncPin() const
diff --git a/Engine/Source/Editor/BlueprintGraph/Private/K2Node_EnumEquality.cpp b/Engine/Source/Editor/BlueprintGraph/Private/K2Node_EnumEquality.cpp
--- a/Engine/Source/Editor/BlueprintGraph/Private/K2Node_EnumEquality.cpp
+++ b/Engine/Source/Editor/BlueprintGraph/Private/K2Node_EnumEquality.cpp
@@ -225,36 +225,33 @@ void UK2Node_EnumEquality::ExpandNode(class FKismetCompilerContext& CompilerCont
 {
 	Super::ExpandNode(CompilerContext, SourceGraph);
 
-	if (CompilerContext.bIsFullCompile)
-	{
-		const UEdGraphSchema_K2* Schema = CompilerContext.GetSchema();
-
-		// Get the enum equality node and the KismetMathLibrary function info for use when we build those nodes
-		FName ConditionalFunctionName = "";
-		UClass* ConditionalFunctionClass = NULL;
-		GetConditionalFunction(ConditionalFunctionName, &ConditionalFunctionClass);
-
-		// Create the conditional node we're replacing the enum node for
-		UK2Node_CallFunction* ConditionalNode = SourceGraph->CreateBlankNode<UK2Node_CallFunction>();
-		ConditionalNode->FunctionReference.SetExternalMember(ConditionalFunctionName, ConditionalFunctionClass);
-		ConditionalNode->AllocateDefaultPins();
-		CompilerContext.MessageLog.NotifyIntermediateObjectCreation(ConditionalNode, this);
-
-		// Rewire the enum pins to the new conditional node
-		UEdGraphPin* LeftSideConditionalPin = ConditionalNode->FindPinChecked(TEXT("A"));
-		UEdGraphPin* RightSideConditionalPin = ConditionalNode->FindPinChecked(TEXT("B"));
-		UEdGraphPin* ReturnConditionalPin = ConditionalNode->FindPinChecked(Schema->PN_ReturnValue);
-		UEdGraphPin* Input1Pin = GetInput1Pin();
-		UEdGraphPin* Input2Pin = GetInput2Pin();
-		LeftSideConditionalPin->PinType = Input1Pin->PinType;
-		RightSideConditionalPin->PinType = Input2Pin->PinType;
-		CompilerContext.MovePinLinksToIntermediate(*Input1Pin, *LeftSideConditionalPin);
-		CompilerContext.MovePinLinksToIntermediate(*Input2Pin, *RightSideConditionalPin);
-		CompilerContext.MovePinLinksToIntermediate(*GetReturnValuePin(), *ReturnConditionalPin);
-
-		// Break all links to the Select node so it goes away for at scheduling time
-		BreakAllNodeLinks();
-	}
+	const UEdGraphSchema_K2* Schema = CompilerContext.GetSchema();
+
+	// Get the enum equality node and the KismetMathLibrary function info for use when we build those nodes
+	FName ConditionalFunctionName = "";
+	UClass* ConditionalFunctionClass = NULL;
+	GetConditionalFunction(ConditionalFunctionName, &ConditionalFunctionClass);
+
+	// Create the conditional node we're replacing the enum node for
+	UK2Node_CallFunction* ConditionalNode = SourceGraph->CreateBlankNode<UK2Node_CallFunction>();
+	ConditionalNode->FunctionReference.SetExternalMember(ConditionalFunctionName, ConditionalFunctionClass);
+	ConditionalNode->AllocateDefaultPins();
+	CompilerContext.MessageLog.NotifyIntermediateObjectCreation(ConditionalNode, this);
+
+	// Rewire the enum pins to the new conditional node
+	UEdGraphPin* LeftSideConditionalPin = ConditionalNode->FindPinChecked(TEXT("A"));
+	UEdGraphPin* RightSideConditionalPin = ConditionalNode->FindPinChecked(TEXT("B"));
+	UEdGraphPin* ReturnConditionalPin = ConditionalNode->FindPinChecked(Schema->PN_ReturnValue);
+	UEdGraphPin* Input1Pin = GetInput1Pin();
+	UEdGraphPin* Input2Pin = GetInput2Pin();
+	LeftSideConditionalPin->PinType = Input1Pin->PinType;
+	RightSideConditionalPin->PinType = Input2Pin->PinType;
+	CompilerContext.MovePinLinksToIntermediate(*Input1Pin, *LeftSideConditionalPin);
+	CompilerContext.MovePinLinksToIntermediate(*Input2Pin, *RightSideConditionalPin);
+	CompilerContext.MovePinLinksToIntermediate(*GetReturnValuePin(), *ReturnConditionalPin);
+
+	// Break all links to the Select node so it goes away for at scheduling time
+	BreakAllNodeLinks();
 }
 
 void UK2Node_EnumEquality::GetMenuActions(FBlueprintActionDatabaseRegistrar& ActionRegistrar) const
diff --git a/Engine/Source/Editor/BlueprintGraph/Private/K2Node_Event.cpp b/Engine/Source/Editor/BlueprintGraph/Private/K2Node_Event.cpp
--- a/Engine/Source/Editor/BlueprintGraph/Private/K2Node_Event.cpp
+++ b/Engine/Source/Editor/BlueprintGraph/Private/K2Node_Event.cpp
@@ -527,29 +527,26 @@ void UK2Node_Event::ExpandNode(class FKismetCompilerContext& CompilerContext, UE
 {
 	Super::ExpandNode(CompilerContext, SourceGraph);
 
-	if (CompilerContext.bIsFullCompile)
+	UEdGraphPin* OrgDelegatePin = FindPin(UK2Node_Event::DelegateOutputName);
+	if (OrgDelegatePin && OrgDelegatePin->LinkedTo.Num() > 0)
 	{
-		UEdGraphPin* OrgDelegatePin = FindPin(UK2Node_Event::DelegateOutputName);
-		if (OrgDelegatePin && OrgDelegatePin->LinkedTo.Num() > 0)
-		{
-			const UEdGraphSchema_K2* Schema = CompilerContext.GetSchema();
+		const UEdGraphSchema_K2* Schema = CompilerContext.GetSchema();
 
-			const FName FunctionName = GetFunctionName();
-			if(FunctionName == NAME_None)
-			{
-				CompilerContext.MessageLog.Error(*FString::Printf(*LOCTEXT("EventDelegateName_Error", "Event node @@ has no name of function.").ToString()), this);
-			}
+		const FName FunctionName = GetFunctionName();
+		if(FunctionName == NAME_None)
+		{
+			CompilerContext.MessageLog.Error(*FString::Printf(*LOCTEXT("EventDelegateName_Error", "Event node @@ has no name of function.").ToString()), this);
+		}
 
-			UK2Node_Self* SelfNode = CompilerContext.SpawnIntermediateNode<UK2Node_Self>(this, SourceGraph);
-			SelfNode->AllocateDefaultPins();
+		UK2Node_Self* SelfNode = CompilerContext.SpawnIntermediateNode<UK2Node_Self>(this, SourceGraph);
+		SelfNode->AllocateDefaultPins();
 
-			UK2Node_CreateDelegate* CreateDelegateNode = CompilerContext.SpawnIntermediateNode<UK2Node_CreateDelegate>(this, SourceGraph);
-			CreateDelegateNode->AllocateDefaultPins();
-			CompilerContext.MovePinLinksToIntermediate(*OrgDelegatePin, *CreateDelegateNode->GetDelegateOutPin());
-			Schema->TryCreateConnection(SelfNode->FindPinChecked(Schema->PN_Self), CreateDelegateNode->GetObjectInPin());
-			// When called UFunction is defined in the same class, it wasn't created yet (previously the Skeletal class was checked). So no "CreateDelegateNode->HandleAnyChangeWithoutNotifying();" is called.
-			CreateDelegateNode->SetFunction(FunctionName);
-		}
+		UK2Node_CreateDelegate* CreateDelegateNode = CompilerContext.SpawnIntermediateNode<UK2Node_CreateDelegate>(this, SourceGraph);
+		CreateDelegateNode->AllocateDefaultPins();
+		CompilerContext.MovePinLinksToIntermediate(*OrgDelegatePin, *CreateDelegateNode->GetDelegateOutPin());
+		Schema->TryCreateConnection(SelfNode->FindPinChecked(Schema->PN_Self), CreateDelegateNode->GetObjectInPin());
+		// When called UFunction is defined in the same class, it wasn't created yet (previously the Skeletal class was checked). So no "CreateDelegateNode->HandleAnyChangeWithoutNotifying();" is called.
+		CreateDelegateNode->SetFunction(FunctionName);
 	}
 }
 
diff --git a/Engine/Source/Editor/BlueprintGraph/Private/K2Node_ForEachElementInEnum.cpp b/Engine/Source/Editor/BlueprintGraph/Private/K2Node_ForEachElementInEnum.cpp
--- a/Engine/Source/Editor/BlueprintGraph/Private/K2Node_ForEachElementInEnum.cpp
+++ b/Engine/Source/Editor/BlueprintGraph/Private/K2Node_ForEachElementInEnum.cpp
@@ -156,51 +156,48 @@ void UK2Node_ForEachElementInEnum::ExpandNode(class FKismetCompilerContext& Comp
 {
 	Super::ExpandNode(CompilerContext, SourceGraph);
 
-	if (CompilerContext.bIsFullCompile)
+	if (!Enum)
 	{
-		if (!Enum)
-		{
-			ValidateNodeDuringCompilation(CompilerContext.MessageLog);
-			return;
-		}
+		ValidateNodeDuringCompilation(CompilerContext.MessageLog);
+		return;
+	}
 
-		FForExpandNodeHelper ForLoop;
-		if (!ForLoop.BuildLoop(this, CompilerContext, SourceGraph))
-		{
-			CompilerContext.MessageLog.Error(*NSLOCTEXT("K2Node", "ForEachElementInEnum_ForError", "For Expand error in @@").ToString(), this);
-		}
+	FForExpandNodeHelper ForLoop;
+	if (!ForLoop.BuildLoop(this, CompilerContext, SourceGraph))
+	{
+		CompilerContext.MessageLog.Error(*NSLOCTEXT("K2Node", "ForEachElementInEnum_ForError", "For Expand error in @@").ToString(), this);
+	}
 
-		const UEdGraphSchema_K2* Schema = CompilerContext.GetSchema();
+	const UEdGraphSchema_K2* Schema = CompilerContext.GetSchema();
 
-		CompilerContext.MovePinLinksToIntermediate(*GetExecPin(), *ForLoop.StartLoopExecInPin);
-		CompilerContext.MovePinLinksToIntermediate(*FindPinChecked(Schema->PN_Then), *ForLoop.LoopCompleteOutExecPin);
-		CompilerContext.MovePinLinksToIntermediate(*FindPinChecked(InsideLoopPinName), *ForLoop.InsideLoopExecOutPin);
-
-		UK2Node_GetNumEnumEntries* GetNumEnumEntries = CompilerContext.SpawnIntermediateNode<UK2Node_GetNumEnumEntries>(this, SourceGraph);
-		GetNumEnumEntries->Enum = Enum;
-		GetNumEnumEntries->AllocateDefaultPins();
-		bool bResult = Schema->TryCreateConnection(GetNumEnumEntries->FindPinChecked(Schema->PN_ReturnValue), ForLoop.IndexLimitInPin);
-
-		UK2Node_CallFunction* Conv_Func = CompilerContext.SpawnIntermediateNode<UK2Node_CallFunction>(this, SourceGraph); 
-		FName Conv_Func_Name = GET_FUNCTION_NAME_CHECKED(UKismetMathLibrary, Conv_IntToByte);
-		Conv_Func->SetFromFunction(UKismetMathLibrary::StaticClass()->FindFunctionByName(Conv_Func_Name));
-		Conv_Func->AllocateDefaultPins();
-		bResult &= Schema->TryCreateConnection(Conv_Func->FindPinChecked(TEXT("InInt")), ForLoop.IndexOutPin);
-
-		UK2Node_CastByteToEnum* CastByteToEnum = CompilerContext.SpawnIntermediateNode<UK2Node_CastByteToEnum>(this, SourceGraph);
-		CastByteToEnum->Enum = Enum;
-		CastByteToEnum->bSafe = true;
-		CastByteToEnum->AllocateDefaultPins();
-		bResult &= Schema->TryCreateConnection(Conv_Func->FindPinChecked(Schema->PN_ReturnValue), CastByteToEnum->FindPinChecked(UK2Node_CastByteToEnum::ByteInputPinName));
-		CompilerContext.MovePinLinksToIntermediate(*FindPinChecked(EnumOuputPinName), *CastByteToEnum->FindPinChecked(Schema->PN_ReturnValue));
-
-		if (!bResult)
-		{
-			CompilerContext.MessageLog.Error(*NSLOCTEXT("K2Node", "ForEachElementInEnum_ExpandError", "Expand error in @@").ToString(), this);
-		}
+	CompilerContext.MovePinLinksToIntermediate(*GetExecPin(), *ForLoop.StartLoopExecInPin);
+	CompilerContext.MovePinLinksToIntermediate(*FindPinChecked(Schema->PN_Then), *ForLoop.LoopCompleteOutExecPin);
+	CompilerContext.MovePinLinksToIntermediate(*FindPinChecked(InsideLoopPinName), *ForLoop.InsideLoopExecOutPin);
+
+	UK2Node_GetNumEnumEntries* GetNumEnumEntries = CompilerContext.SpawnIntermediateNode<UK2Node_GetNumEnumEntries>(this, SourceGraph);
+	GetNumEnumEntries->Enum = Enum;
+	GetNumEnumEntries->AllocateDefaultPins();
+	bool bResult = Schema->TryCreateConnection(GetNumEnumEntries->FindPinChecked(Schema->PN_ReturnValue), ForLoop.IndexLimitInPin);
 
-		BreakAllNodeLinks();
+	UK2Node_CallFunction* Conv_Func = CompilerContext.SpawnIntermediateNode<UK2Node_CallFunction>(this, SourceGraph);
+	FName Conv_Func_Name = GET_FUNCTION_NAME_CHECKED(UKismetMathLibrary, Conv_IntToByte);
+	Conv_Func->SetFromFunction(UKismetMathLibrary::StaticClass()->FindFunctionByName(Conv_Func_Name));
+	Conv_Func->AllocateDefaultPins();
+	bResult &= Schema->TryCreateConnection(Conv_Func->FindPinChecked(TEXT("InInt")), ForLoop.IndexOutPin);
+
+	UK2Node_CastByteToEnum* CastByteToEnum = CompilerContext.SpawnIntermediateNode<UK2Node_CastByteToEnum>(this, SourceGraph);
+	CastByteToEnum->Enum = Enum;
+	CastByteToEnum->bSafe = true;
+	CastByteToEnum->AllocateDefaultPins();
+	bResult &= Schema->TryCreateConnection(Conv_Func->FindPinChecked(Schema->PN_ReturnValue), CastByteToEnum->FindPinChecked(UK2Node_CastByteToEnum::ByteInputPinName));
+	CompilerContext.MovePinLinksToIntermediate(*FindPinChecked(EnumOuputPinName), *CastByteToEnum->FindPinChecked(Schema->PN_ReturnValue));
+
+	if (!bResult)
+	{
+		CompilerContext.MessageLog.Error(*NSLOCTEXT("K2Node", "ForEachElementInEnum_ExpandError", "Expand error in @@").ToString(), this);
 	}
+
+	BreakAllNodeLinks();
 }
 
 void UK2Node_ForEachElementInEnum::GetMenuActions(FBlueprintActionDatabaseRegistrar& ActionRegistrar) const
diff --git a/Engine/Source/Editor/BlueprintGraph/Private/K2Node_FormatText.cpp b/Engine/Source/Editor/BlueprintGraph/Private/K2Node_FormatText.cpp
--- a/Engine/Source/Editor/BlueprintGraph/Private/K2Node_FormatText.cpp
+++ b/Engine/Source/Editor/BlueprintGraph/Private/K2Node_FormatText.cpp
@@ -183,73 +183,69 @@ void UK2Node_FormatText::ExpandNode(class FKismetCompilerContext& CompilerContex
 {
 	Super::ExpandNode(CompilerContext, SourceGraph);
 
-	if (CompilerContext.bIsFullCompile)
-	{
-		/**
-			At the end of this, the UK2Node_FormatText will not be a part of the Blueprint, it merely handles connecting
-			the other nodes into the Blueprint.
-		*/
-
-		const UEdGraphSchema_K2* Schema = CompilerContext.GetSchema();
+	/**
+		At the end of this, the UK2Node_FormatText will not be a part of the Blueprint, it merely handles connecting
+		the other nodes into the Blueprint.
+	*/
 
-		// Create a "Make Array" node to compile the list of arguments into an array for the Format function being called
-		UK2Node_MakeArray* MakeArrayNode = CompilerContext.SpawnIntermediateNode<UK2Node_MakeArray>(this, SourceGraph); //SourceGraph->CreateBlankNode<UK2Node_MakeArray>();
-		MakeArrayNode->AllocateDefaultPins();
-		CompilerContext.MessageLog.NotifyIntermediateObjectCreation(MakeArrayNode, this);
+	const UEdGraphSchema_K2* Schema = CompilerContext.GetSchema();
 
-		UEdGraphPin* ArrayOut = MakeArrayNode->GetOutputPin();
+	// Create a "Make Array" node to compile the list of arguments into an array for the Format function being called
+	UK2Node_MakeArray* MakeArrayNode = CompilerContext.SpawnIntermediateNode<UK2Node_MakeArray>(this, SourceGraph); //SourceGraph->CreateBlankNode<UK2Node_MakeArray>();
+	MakeArrayNode->AllocateDefaultPins();
+	CompilerContext.MessageLog.NotifyIntermediateObjectCreation(MakeArrayNode, this);
 
-		// This is the node that does all the Format work.
-		UK2Node_CallFunction* CallFunction = CompilerContext.SpawnIntermediateNode<UK2Node_CallFunction>(this, SourceGraph);
-		CallFunction->SetFromFunction(UKismetTextLibrary::StaticClass()->FindFunctionByName("Format"));
-		CallFunction->AllocateDefaultPins();
-		CompilerContext.MessageLog.NotifyIntermediateObjectCreation(CallFunction, this);
+	UEdGraphPin* ArrayOut = MakeArrayNode->GetOutputPin();
 
-		// Connect the output of the "Make Array" pin to the function's "InArgs" pin
-		ArrayOut->MakeLinkTo(CallFunction->FindPin(TEXT("InArgs")));
+	// This is the node that does all the Format work.
+	UK2Node_CallFunction* CallFunction = CompilerContext.SpawnIntermediateNode<UK2Node_CallFunction>(this, SourceGraph);
+	CallFunction->SetFromFunction(UKismetTextLibrary::StaticClass()->FindFunctionByName("Format"));
+	CallFunction->AllocateDefaultPins();
+	CompilerContext.MessageLog.NotifyIntermediateObjectCreation(CallFunction, this);
 
-		// This will set the "Make Array" node's type, only works if one pin is connected.
-		MakeArrayNode->PinConnectionListChanged(ArrayOut);
+	// Connect the output of the "Make Array" pin to the function's "InArgs" pin
+	ArrayOut->MakeLinkTo(CallFunction->FindPin(TEXT("InArgs")));
 
-		// For each argument, we will need to add in a "Make Struct" node.
-		for(int32 ArgIdx = 0; ArgIdx < PinNames.Num(); ++ArgIdx)
-		{
-			UEdGraphPin* ArgumentPin = FindArgumentPin(PinNames[ArgIdx]);
+	// This will set the "Make Array" node's type, only works if one pin is connected.
+	MakeArrayNode->PinConnectionListChanged(ArrayOut);
 
-			// Spawn a "Make Struct" node to create the struct needed for formatting the text.
-			UK2Node_MakeStruct* PinMakeStruct = CompilerContext.SpawnIntermediateNode<UK2Node_MakeStruct>(this, SourceGraph);
-			PinMakeStruct->StructType = FFormatTextArgument::StaticStruct();
-			PinMakeStruct->AllocateDefaultPins();
+	// For each argument, we will need to add in a "Make Struct" node.
+	for(int32 ArgIdx = 0; ArgIdx < PinNames.Num(); ++ArgIdx)
+	{
+		UEdGraphPin* ArgumentPin = FindArgumentPin(PinNames[ArgIdx]);
 
-			// Set the struct's "ArgumentName" pin literal to be the argument pin's name.
-			PinMakeStruct->GetSchema()->TrySetDefaultText(*PinMakeStruct->FindPin("ArgumentName"), FText::AsCultureInvariant(ArgumentPin->PinName));
+		// Spawn a "Make Struct" node to create the struct needed for formatting the text.
+		UK2Node_MakeStruct* PinMakeStruct = CompilerContext.SpawnIntermediateNode<UK2Node_MakeStruct>(this, SourceGraph);
+		PinMakeStruct->StructType = FFormatTextArgument::StaticStruct();
+		PinMakeStruct->AllocateDefaultPins();
 
-			// Move the connection of the argument pin to the struct's "TextValue" pin, this will move the literal value if present.
-			CompilerContext.MovePinLinksToIntermediate(*ArgumentPin, *PinMakeStruct->FindPin("TextValue"));
+		// Set the struct's "ArgumentName" pin literal to be the argument pin's name.
+		PinMakeStruct->GetSchema()->TrySetDefaultText(*PinMakeStruct->FindPin("ArgumentName"), FText::AsCultureInvariant(ArgumentPin->PinName));
 
-			// The "Make Array" node already has one pin available, so don't create one for ArgIdx == 0
-			if(ArgIdx > 0)
-			{
-				MakeArrayNode->AddInputPin();
-			}
-
-			// Find the input pin on the "Make Array" node by index.
-			FString PinName = FString::Printf(TEXT("[%d]"), ArgIdx);
-			UEdGraphPin* InputPin = MakeArrayNode->FindPin(PinName);
-
-			// Find the output for the pin's "Make Struct" node and link it to the corresponding pin on the "Make Array" node.
-			FindOutputStructPinChecked(PinMakeStruct)->MakeLinkTo(InputPin);
+		// Move the connection of the argument pin to the struct's "TextValue" pin, this will move the literal value if present.
+		CompilerContext.MovePinLinksToIntermediate(*ArgumentPin, *PinMakeStruct->FindPin("TextValue"));
 
+		// The "Make Array" node already has one pin available, so don't create one for ArgIdx == 0
+		if(ArgIdx > 0)
+		{
+			MakeArrayNode->AddInputPin();
 		}
 
-		// Move connection of FormatText's "Result" pin to the call function's return value pin.
-		CompilerContext.MovePinLinksToIntermediate(*FindPin(TEXT("Result")), *CallFunction->GetReturnValuePin());
-		// Move connection of FormatText's "Format" pin to the call function's "InPattern" pin
-		CompilerContext.MovePinLinksToIntermediate(*GetFormatPin(), *CallFunction->FindPin(TEXT("InPattern")));
+		// Find the input pin on the "Make Array" node by index.
+		FString PinName = FString::Printf(TEXT("[%d]"), ArgIdx);
+		UEdGraphPin* InputPin = MakeArrayNode->FindPin(PinName);
+
+		// Find the output for the pin's "Make Struct" node and link it to the corresponding pin on the "Make Array" node.
+		FindOutputStructPinChecked(PinMakeStruct)->MakeLinkTo(InputPin);
 
-		BreakAllNodeLinks();
 	}
 
+	// Move connection of FormatText's "Result" pin to the call function's return value pin.
+	CompilerContext.MovePinLinksToIntermediate(*FindPin(TEXT("Result")), *CallFunction->GetReturnValuePin());
+	// Move connection of FormatText's "Format" pin to the call function's "InPattern" pin
+	CompilerContext.MovePinLinksToIntermediate(*GetFormatPin(), *CallFunction->FindPin(TEXT("InPattern")));
+
+	BreakAllNodeLinks();
 }
 
 UEdGraphPin* UK2Node_FormatText::FindArgumentPin(const FText& InPinName) const
diff --git a/Engine/Source/Editor/BlueprintGraph/Private/K2Node_GetDataTableRow.cpp b/Engine/Source/Editor/BlueprintGraph/Private/K2Node_GetDataTableRow.cpp
--- a/Engine/Source/Editor/BlueprintGraph/Private/K2Node_GetDataTableRow.cpp
+++ b/Engine/Source/Editor/BlueprintGraph/Private/K2Node_GetDataTableRow.cpp
@@ -314,66 +314,63 @@ void UK2Node_GetDataTableRow::ExpandNode(class FKismetCompilerContext& CompilerC
 {
     Super::ExpandNode(CompilerContext, SourceGraph);
     
-	if (CompilerContext.bIsFullCompile)
-	{
-		const UEdGraphSchema_K2* Schema = CompilerContext.GetSchema();
+	const UEdGraphSchema_K2* Schema = CompilerContext.GetSchema();
         
-        UEdGraphPin* OriginalDataTableInPin = GetDataTablePin();
-        UDataTable* Table = (OriginalDataTableInPin != NULL) ? Cast<UDataTable>(OriginalDataTableInPin->DefaultObject) : NULL;
-        if((0 == OriginalDataTableInPin->LinkedTo.Num()) && (NULL == Table))
-        {
-        	CompilerContext.MessageLog.Error(*LOCTEXT("GetDataTableRowNoDataTable_Error", "GetDataTableRow must have a DataTable specified.").ToString(), this);
-        	// we break exec links so this is the only error we get
-            BreakAllNodeLinks();
-            return;
-        }
-
-		// FUNCTION NODE
-		const FName FunctionName = GET_FUNCTION_NAME_CHECKED(UDataTableFunctionLibrary, GetDataTableRowFromName);
-		UK2Node_CallFunction* GetDataTableRowFunction = CompilerContext.SpawnIntermediateNode<UK2Node_CallFunction>(this, SourceGraph);
-		GetDataTableRowFunction->FunctionReference.SetExternalMember(FunctionName, UDataTableFunctionLibrary::StaticClass());
-		GetDataTableRowFunction->AllocateDefaultPins();
-        CompilerContext.MovePinLinksToIntermediate(*GetExecPin(), *(GetDataTableRowFunction->GetExecPin()));
-
-		// Connect the input of our GetDataTableRow to the Input of our Function pin
-        UEdGraphPin* DataTableInPin = GetDataTableRowFunction->FindPinChecked(TEXT("Table"));
-		if(OriginalDataTableInPin->LinkedTo.Num() > 0)
-		{
-			// Copy the connection
-			CompilerContext.MovePinLinksToIntermediate(*OriginalDataTableInPin, *DataTableInPin);
-		}
-		else
-		{
-			// Copy literal
-			DataTableInPin->DefaultObject = OriginalDataTableInPin->DefaultObject;
-		}
-		UEdGraphPin* RowNameInPin = GetDataTableRowFunction->FindPinChecked(TEXT("RowName"));
-		CompilerContext.MovePinLinksToIntermediate(*GetRowNamePin(), *RowNameInPin);
-
-		// Get some pins to work with
-		UEdGraphPin* OriginalOutRowPin = FindPinChecked(Schema->PN_ReturnValue);
-		UEdGraphPin* FunctionOutRowPin = GetDataTableRowFunction->FindPinChecked(TEXT("OutRow"));
-        UEdGraphPin* FunctionReturnPin = GetDataTableRowFunction->FindPinChecked(Schema->PN_ReturnValue);
-        UEdGraphPin* FunctionThenPin = GetDataTableRowFunction->GetThenPin();
+    UEdGraphPin* OriginalDataTableInPin = GetDataTablePin();
+    UDataTable* Table = (OriginalDataTableInPin != NULL) ? Cast<UDataTable>(OriginalDataTableInPin->DefaultObject) : NULL;
+    if((0 == OriginalDataTableInPin->LinkedTo.Num()) && (NULL == Table))
+    {
+        CompilerContext.MessageLog.Error(*LOCTEXT("GetDataTableRowNoDataTable_Error", "GetDataTableRow must have a DataTable specified.").ToString(), this);
+        // we break exec links so this is the only error we get
+        BreakAllNodeLinks();
+        return;
+    }
+
+	// FUNCTION NODE
+	const FName FunctionName = GET_FUNCTION_NAME_CHECKED(UDataTableFunctionLibrary, GetDataTableRowFromName);
+	UK2Node_CallFunction* GetDataTableRowFunction = CompilerContext.SpawnIntermediateNode<UK2Node_CallFunction>(this, SourceGraph);
+	GetDataTableRowFunction->FunctionReference.SetExternalMember(FunctionName, UDataTableFunctionLibrary::StaticClass());
+	GetDataTableRowFunction->AllocateDefaultPins();
+    CompilerContext.MovePinLinksToIntermediate(*GetExecPin(), *(GetDataTableRowFunction->GetExecPin()));
+
+	// Connect the input of our GetDataTableRow to the Input of our Function pin
+    UEdGraphPin* DataTableInPin = GetDataTableRowFunction->FindPinChecked(TEXT("Table"));
+	if(OriginalDataTableInPin->LinkedTo.Num() > 0)
+	{
+		// Copy the connection
+		CompilerContext.MovePinLinksToIntermediate(*OriginalDataTableInPin, *DataTableInPin);
+	}
+	else
+	{
+		// Copy literal
+		DataTableInPin->DefaultObject = OriginalDataTableInPin->DefaultObject;
+	}
+	UEdGraphPin* RowNameInPin = GetDataTableRowFunction->FindPinChecked(TEXT("RowName"));
+	CompilerContext.MovePinLinksToIntermediate(*GetRowNamePin(), *RowNameInPin);
+
+	// Get some pins to work with
+	UEdGraphPin* OriginalOutRowPin = FindPinChecked(Schema->PN_ReturnValue);
+	UEdGraphPin* FunctionOutRowPin = GetDataTableRowFunction->FindPinChecked(TEXT("OutRow"));
+    UEdGraphPin* FunctionReturnPin = GetDataTableRowFunction->FindPinChecked(Schema->PN_ReturnValue);
+    UEdGraphPin* FunctionThenPin = GetDataTableRowFunction->GetThenPin();
         
-        // Set the type of the OutRow pin on this expanded mode to match original
-        FunctionOutRowPin->PinType = OriginalOutRowPin->PinType;
-		FunctionOutRowPin->PinType.PinSubCategoryObject = OriginalOutRowPin->PinType.PinSubCategoryObject;
+    // Set the type of the OutRow pin on this expanded mode to match original
+    FunctionOutRowPin->PinType = OriginalOutRowPin->PinType;
+	FunctionOutRowPin->PinType.PinSubCategoryObject = OriginalOutRowPin->PinType.PinSubCategoryObject;
         
-        //BRANCH NODE
-        UK2Node_IfThenElse* BranchNode = CompilerContext.SpawnIntermediateNode<UK2Node_IfThenElse>(this, SourceGraph);
-        BranchNode->AllocateDefaultPins();
-        // Hook up inputs to branch
-        FunctionThenPin->MakeLinkTo(BranchNode->GetExecPin());
-        FunctionReturnPin->MakeLinkTo(BranchNode->GetConditionPin());
+    //BRANCH NODE
+    UK2Node_IfThenElse* BranchNode = CompilerContext.SpawnIntermediateNode<UK2Node_IfThenElse>(this, SourceGraph);
+    BranchNode->AllocateDefaultPins();
+    // Hook up inputs to branch
+    FunctionThenPin->MakeLinkTo(BranchNode->GetExecPin());
+    FunctionReturnPin->MakeLinkTo(BranchNode->GetConditionPin());
         
-        // Hook up outputs
-        CompilerContext.MovePinLinksToIntermediate(*GetThenPin(), *(BranchNode->GetThenPin()));
-        CompilerContext.MovePinLinksToIntermediate(*GetRowNotFoundPin(), *(BranchNode->GetElsePin()));
-        CompilerContext.MovePinLinksToIntermediate(*OriginalOutRowPin, *FunctionOutRowPin);
+    // Hook up outputs
+    CompilerContext.MovePinLinksToIntermediate(*GetThenPin(), *(BranchNode->GetThenPin()));
+    CompilerContext.MovePinLinksToIntermediate(*GetRowNotFoundPin(), *(BranchNode->GetElsePin()));
+    CompilerContext.MovePinLinksToIntermediate(*OriginalOutRowPin, *FunctionOutRowPin);
 
-		BreakAllNodeLinks();
-	}
+	BreakAllNodeLinks();
 }
 
 #undef LOCTEXT_NAMESPACE
\ No newline at end of file
diff --git a/Engine/Source/Editor/BlueprintGraph/Private/K2Node_GetEnumeratorName.cpp b/Engine/Source/Editor/BlueprintGraph/Private/K2Node_GetEnumeratorName.cpp
--- a/Engine/Source/Editor/BlueprintGraph/Private/K2Node_GetEnumeratorName.cpp
+++ b/Engine/Source/Editor/BlueprintGraph/Private/K2Node_GetEnumeratorName.cpp
@@ -121,68 +121,65 @@ void UK2Node_GetEnumeratorName::ExpandNode(class FKismetCompilerContext& Compile
 {
 	Super::ExpandNode(CompilerContext, SourceGraph);
 
-	if (CompilerContext.bIsFullCompile)
+	UEnum* Enum = GetEnum();
+	if(NULL == Enum)
 	{
-		UEnum* Enum = GetEnum();
-		if(NULL == Enum)
-		{
-			CompilerContext.MessageLog.Error(*FString::Printf(*NSLOCTEXT("K2Node", "GetEnumeratorNam_Error_MustHaveValidName", "@@ must have a valid enum defined").ToString()), this);
-			return;
-		}
+		CompilerContext.MessageLog.Error(*FString::Printf(*NSLOCTEXT("K2Node", "GetEnumeratorNam_Error_MustHaveValidName", "@@ must have a valid enum defined").ToString()), this);
+		return;
+	}
 
-		const UEdGraphSchema_K2* Schema = CompilerContext.GetSchema();
+	const UEdGraphSchema_K2* Schema = CompilerContext.GetSchema();
 		
-		const UFunction* Function = UKismetNodeHelperLibrary::StaticClass()->FindFunctionByName( GetFunctionName() );
-		check(NULL != Function);
-		UK2Node_CallFunction* CallGetName = CompilerContext.SpawnIntermediateNode<UK2Node_CallFunction>(this, SourceGraph); 
-		CallGetName->SetFromFunction(Function);
-		CallGetName->AllocateDefaultPins();
-		check(CallGetName->IsNodePure());
+	const UFunction* Function = UKismetNodeHelperLibrary::StaticClass()->FindFunctionByName( GetFunctionName() );
+	check(NULL != Function);
+	UK2Node_CallFunction* CallGetName = CompilerContext.SpawnIntermediateNode<UK2Node_CallFunction>(this, SourceGraph); 
+	CallGetName->SetFromFunction(Function);
+	CallGetName->AllocateDefaultPins();
+	check(CallGetName->IsNodePure());
 		
-		//OPUTPUT PIN
-		UEdGraphPin* OrgReturnPin = FindPinChecked(Schema->PN_ReturnValue);
-		UEdGraphPin* NewReturnPin = CallGetName->GetReturnValuePin();
-		check(NULL != NewReturnPin);
-		CompilerContext.MovePinLinksToIntermediate(*OrgReturnPin, *NewReturnPin);
-
-		//ENUM PIN
-		UEdGraphPin* EnumPin = CallGetName->FindPinChecked(TEXT("Enum"));
-		Schema->TrySetDefaultObject(*EnumPin, Enum);
-		check(EnumPin->DefaultObject == Enum);
-
-		//INDEX PIN
-		UEdGraphPin* OrgInputPin = FindPinChecked(EnumeratorPinName);
-		UEdGraphPin* IndexPin = CallGetName->FindPinChecked(TEXT("EnumeratorIndex"));
-		check(EGPD_Input == IndexPin->Direction && Schema->PC_Byte == IndexPin->PinType.PinCategory);
-		CompilerContext.MovePinLinksToIntermediate(*OrgInputPin, *IndexPin);
-
-		if (!IndexPin->LinkedTo.Num())
+	//OPUTPUT PIN
+	UEdGraphPin* OrgReturnPin = FindPinChecked(Schema->PN_ReturnValue);
+	UEdGraphPin* NewReturnPin = CallGetName->GetReturnValuePin();
+	check(NULL != NewReturnPin);
+	CompilerContext.MovePinLinksToIntermediate(*OrgReturnPin, *NewReturnPin);
+
+	//ENUM PIN
+	UEdGraphPin* EnumPin = CallGetName->FindPinChecked(TEXT("Enum"));
+	Schema->TrySetDefaultObject(*EnumPin, Enum);
+	check(EnumPin->DefaultObject == Enum);
+
+	//INDEX PIN
+	UEdGraphPin* OrgInputPin = FindPinChecked(EnumeratorPinName);
+	UEdGraphPin* IndexPin = CallGetName->FindPinChecked(TEXT("EnumeratorIndex"));
+	check(EGPD_Input == IndexPin->Direction && Schema->PC_Byte == IndexPin->PinType.PinCategory);
+	CompilerContext.MovePinLinksToIntermediate(*OrgInputPin, *IndexPin);
+
+	if (!IndexPin->LinkedTo.Num())
+	{
+		//MAKE LITERAL BYTE FROM LITERAL ENUM
+		const FString EnumLiteral = IndexPin->GetDefaultAsString();
+		const int32 NumericValue = Enum->FindEnumIndex(*EnumLiteral);
+		if (NumericValue == INDEX_NONE) 
 		{
-			//MAKE LITERAL BYTE FROM LITERAL ENUM
-			const FString EnumLiteral = IndexPin->GetDefaultAsString();
-			const int32 NumericValue = Enum->FindEnumIndex(*EnumLiteral);
-			if (NumericValue == INDEX_NONE) 
-			{
-				CompilerContext.MessageLog.Error(*FString::Printf(*NSLOCTEXT("K2Node", "GetEnumeratorNam_Error_InvalidName", "@@ has invalid enum value '%s'").ToString(), *EnumLiteral), this);
-				return;
-			}
-			const FString DefaultByteValue = FString::FromInt(NumericValue);
-
-			// LITERAL BYTE FUNCTION
-			const FName FunctionName = GET_FUNCTION_NAME_CHECKED(UKismetSystemLibrary, MakeLiteralByte);
-			UK2Node_CallFunction* MakeLiteralByte = CompilerContext.SpawnIntermediateNode<UK2Node_CallFunction>(this, SourceGraph); 
-			MakeLiteralByte->SetFromFunction(UKismetSystemLibrary::StaticClass()->FindFunctionByName(FunctionName));
-			MakeLiteralByte->AllocateDefaultPins();
-
-			UEdGraphPin* MakeLiteralByteReturnPin = MakeLiteralByte->FindPinChecked(Schema->PN_ReturnValue);
-			Schema->TryCreateConnection(MakeLiteralByteReturnPin, IndexPin);
-
-			UEdGraphPin* MakeLiteralByteInputPin = MakeLiteralByte->FindPinChecked(TEXT("Value"));
-			MakeLiteralByteInputPin->DefaultValue = DefaultByteValue;
+			CompilerContext.MessageLog.Error(*FString::Printf(*NSLOCTEXT("K2Node", "GetEnumeratorNam_Error_InvalidName", "@@ has invalid enum value '%s'").ToString(), *EnumLiteral), this);
+			return;
 		}
+		const FString DefaultByteValue = FString::FromInt(NumericValue);
+
+		// LITERAL BYTE FUNCTION
+		const FName FunctionName = GET_FUNCTION_NAME_CHECKED(UKismetSystemLibrary, MakeLiteralByte);
+		UK2Node_CallFunction* MakeLiteralByte = CompilerContext.SpawnIntermediateNode<UK2Node_CallFunction>(this, SourceGraph); 
+		MakeLiteralByte->SetFromFunction(UKismetSystemLibrary::StaticClass()->FindFunctionByName(FunctionName));
+		MakeLiteralByte->AllocateDefaultPins();
 
-		BreakAllNodeLinks();
+		UEdGraphPin* MakeLiteralByteReturnPin = MakeLiteralByte->FindPinChecked(Schema->PN_ReturnValue);
+		Schema->TryCreateConnection(MakeLiteralByteReturnPin, IndexPin);
+
+		UEdGraphPin* MakeLiteralByteInputPin = MakeLiteralByte->FindPinChecked(TEXT("Value"));
+		MakeLiteralByteInputPin->DefaultValue = DefaultByteValue;
 	}
+
+	BreakAllNodeLinks();
 }
 
 void UK2Node_GetEnumeratorName::GetMenuActions(FBlueprintActionDatabaseRegistrar& ActionRegistrar) const
diff --git a/Engine/Source/Editor/BlueprintGraph/Private/K2Node_GetNumEnumEntries.cpp b/Engine/Source/Editor/BlueprintGraph/Private/K2Node_GetNumEnumEntries.cpp
--- a/Engine/Source/Editor/BlueprintGraph/Private/K2Node_GetNumEnumEntries.cpp
+++ b/Engine/Source/Editor/BlueprintGraph/Private/K2Node_GetNumEnumEntries.cpp
@@ -57,43 +57,40 @@ void UK2Node_GetNumEnumEntries::ExpandNode(class FKismetCompilerContext& Compile
 {
 	Super::ExpandNode(CompilerContext, SourceGraph);
 
-	if (CompilerContext.bIsFullCompile)
+	if(NULL == Enum)
 	{
-		if(NULL == Enum)
-		{
-			CompilerContext.MessageLog.Error(*FString::Printf(*NSLOCTEXT("K2Node", "GetNumEnumEntries_Error", "@@ must have a valid enum defined").ToString()), this);
-			return;
-		}
+		CompilerContext.MessageLog.Error(*FString::Printf(*NSLOCTEXT("K2Node", "GetNumEnumEntries_Error", "@@ must have a valid enum defined").ToString()), this);
+		return;
+	}
 
-		// Force the enum to load its values if it hasn't already
-		if (Enum->HasAnyFlags(RF_NeedLoad))
-		{
-			Enum->GetLinker()->Preload(Enum);
-		}
+	// Force the enum to load its values if it hasn't already
+	if (Enum->HasAnyFlags(RF_NeedLoad))
+	{
+		Enum->GetLinker()->Preload(Enum);
+	}
 
-		const UEdGraphSchema_K2* Schema = CompilerContext.GetSchema();
-		check(NULL != Schema);
+	const UEdGraphSchema_K2* Schema = CompilerContext.GetSchema();
+	check(NULL != Schema);
 
-		//MAKE LITERAL
-		const FName FunctionName = GET_FUNCTION_NAME_CHECKED(UKismetSystemLibrary, MakeLiteralInt);
-		UK2Node_CallFunction* MakeLiteralInt = CompilerContext.SpawnIntermediateNode<UK2Node_CallFunction>(this, SourceGraph); 
-		MakeLiteralInt->SetFromFunction(UKismetSystemLibrary::StaticClass()->FindFunctionByName(FunctionName));
-		MakeLiteralInt->AllocateDefaultPins();
+	//MAKE LITERAL
+	const FName FunctionName = GET_FUNCTION_NAME_CHECKED(UKismetSystemLibrary, MakeLiteralInt);
+	UK2Node_CallFunction* MakeLiteralInt = CompilerContext.SpawnIntermediateNode<UK2Node_CallFunction>(this, SourceGraph); 
+	MakeLiteralInt->SetFromFunction(UKismetSystemLibrary::StaticClass()->FindFunctionByName(FunctionName));
+	MakeLiteralInt->AllocateDefaultPins();
 
-		//OPUTPUT PIN
-		UEdGraphPin* OrgReturnPin = FindPinChecked(Schema->PN_ReturnValue);
-		UEdGraphPin* NewReturnPin = MakeLiteralInt->GetReturnValuePin();
-		check(NULL != NewReturnPin);
-		CompilerContext.MovePinLinksToIntermediate(*OrgReturnPin, *NewReturnPin);
+	//OPUTPUT PIN
+	UEdGraphPin* OrgReturnPin = FindPinChecked(Schema->PN_ReturnValue);
+	UEdGraphPin* NewReturnPin = MakeLiteralInt->GetReturnValuePin();
+	check(NULL != NewReturnPin);
+	CompilerContext.MovePinLinksToIntermediate(*OrgReturnPin, *NewReturnPin);
 
-		//INPUT PIN
-		UEdGraphPin* InputPin = MakeLiteralInt->FindPinChecked(TEXT("Value"));
-		check(EGPD_Input == InputPin->Direction);
-		const FString DefaultValue = FString::FromInt(Enum->NumEnums() - 1);
-		InputPin->DefaultValue = DefaultValue;
+	//INPUT PIN
+	UEdGraphPin* InputPin = MakeLiteralInt->FindPinChecked(TEXT("Value"));
+	check(EGPD_Input == InputPin->Direction);
+	const FString DefaultValue = FString::FromInt(Enum->NumEnums() - 1);
+	InputPin->DefaultValue = DefaultValue;
 
-		BreakAllNodeLinks();
-	}
+	BreakAllNodeLinks();
 }
 
 void UK2Node_GetNumEnumEntries::GetMenuActions(FBlueprintActionDatabaseRegistrar& ActionRegistrar) const
diff --git a/Engine/Source/Editor/BlueprintGraph/Private/K2Node_InputAction.cpp b/Engine/Source/Editor/BlueprintGraph/Private/K2Node_InputAction.cpp
--- a/Engine/Source/Editor/BlueprintGraph/Private/K2Node_InputAction.cpp
+++ b/Engine/Source/Editor/BlueprintGraph/Private/K2Node_InputAction.cpp
@@ -137,11 +137,8 @@ void UK2Node_InputAction::ExpandNode(FKismetCompilerContext& CompilerContext, UE
 {
 	Super::ExpandNode(CompilerContext, SourceGraph);
 
-	if (CompilerContext.bIsFullCompile)
-	{
-		CreateInputActionEvent(CompilerContext, SourceGraph, GetPressedPin(), IE_Pressed);
-		CreateInputActionEvent(CompilerContext, SourceGraph, GetReleasedPin(), IE_Released);
-	}
+	CreateInputActionEvent(CompilerContext, SourceGraph, GetPressedPin(), IE_Pressed);
+	CreateInputActionEvent(CompilerContext, SourceGraph, GetReleasedPin(), IE_Released);
 }
 
 void UK2Node_InputAction::GetMenuActions(FBlueprintActionDatabaseRegistrar& ActionRegistrar) const
diff --git a/Engine/Source/Editor/BlueprintGraph/Private/K2Node_InputKey.cpp b/Engine/Source/Editor/BlueprintGraph/Private/K2Node_InputKey.cpp
--- a/Engine/Source/Editor/BlueprintGraph/Private/K2Node_InputKey.cpp
+++ b/Engine/Source/Editor/BlueprintGraph/Private/K2Node_InputKey.cpp
@@ -364,11 +364,8 @@ void UK2Node_InputKey::ExpandNode(FKismetCompilerContext& CompilerContext, UEdGr
 {
 	Super::ExpandNode(CompilerContext, SourceGraph);
 
-	if (CompilerContext.bIsFullCompile)
-	{
-		CreateInputKeyEvent(CompilerContext, SourceGraph, GetPressedPin(), IE_Pressed);
-		CreateInputKeyEvent(CompilerContext, SourceGraph, GetReleasedPin(), IE_Released);
-	}
+	CreateInputKeyEvent(CompilerContext, SourceGraph, GetPressedPin(), IE_Pressed);
+	CreateInputKeyEvent(CompilerContext, SourceGraph, GetReleasedPin(), IE_Released);
 }
 
 void UK2Node_InputKey::GetMenuActions(FBlueprintActionDatabaseRegistrar& ActionRegistrar) const
diff --git a/Engine/Source/Editor/BlueprintGraph/Private/K2Node_InputTouch.cpp b/Engine/Source/Editor/BlueprintGraph/Private/K2Node_InputTouch.cpp
--- a/Engine/Source/Editor/BlueprintGraph/Private/K2Node_InputTouch.cpp
+++ b/Engine/Source/Editor/BlueprintGraph/Private/K2Node_InputTouch.cpp
@@ -127,123 +127,120 @@ void UK2Node_InputTouch::ExpandNode(FKismetCompilerContext& CompilerContext, UEd
 {
 	Super::ExpandNode(CompilerContext, SourceGraph);
 
-	if (CompilerContext.bIsFullCompile)
-	{
-		UEdGraphPin* InputTouchPressedPin = GetPressedPin();
-		UEdGraphPin* InputTouchReleasedPin = GetReleasedPin();
+	UEdGraphPin* InputTouchPressedPin = GetPressedPin();
+	UEdGraphPin* InputTouchReleasedPin = GetReleasedPin();
 
-		const UEdGraphSchema_K2* Schema = CompilerContext.GetSchema();
+	const UEdGraphSchema_K2* Schema = CompilerContext.GetSchema();
 
-		// If both are linked we have to do more complicated behaviors
-		if (InputTouchPressedPin->LinkedTo.Num() > 0 && InputTouchReleasedPin->LinkedTo.Num() > 0)
+	// If both are linked we have to do more complicated behaviors
+	if (InputTouchPressedPin->LinkedTo.Num() > 0 && InputTouchReleasedPin->LinkedTo.Num() > 0)
+	{
+		// Create the input touch events
+		UK2Node_InputTouchEvent* InputTouchPressedEvent = CompilerContext.SpawnIntermediateNode<UK2Node_InputTouchEvent>(this, SourceGraph);
+		InputTouchPressedEvent->CustomFunctionName = FName( *FString::Printf(TEXT("InpTchEvt_%s"), *InputTouchPressedEvent->GetName()));
+		InputTouchPressedEvent->bConsumeInput = bConsumeInput;
+		InputTouchPressedEvent->bExecuteWhenPaused = bExecuteWhenPaused;
+		InputTouchPressedEvent->bOverrideParentBinding = bOverrideParentBinding;
+		InputTouchPressedEvent->InputKeyEvent = IE_Pressed;
+		InputTouchPressedEvent->EventSignatureName = TEXT("InputTouchHandlerDynamicSignature__DelegateSignature");
+		InputTouchPressedEvent->EventSignatureClass = UInputComponent::StaticClass();
+		InputTouchPressedEvent->bInternalEvent = true;
+		InputTouchPressedEvent->AllocateDefaultPins();
+
+		UK2Node_InputTouchEvent* InputTouchReleasedEvent = CompilerContext.SpawnIntermediateNode<UK2Node_InputTouchEvent>(this, SourceGraph);
+		InputTouchReleasedEvent->CustomFunctionName = FName( *FString::Printf(TEXT("InpTchEvt_%s"),  *InputTouchReleasedEvent->GetName()));
+		InputTouchReleasedEvent->bConsumeInput = bConsumeInput;
+		InputTouchReleasedEvent->bExecuteWhenPaused = bExecuteWhenPaused;
+		InputTouchReleasedEvent->bOverrideParentBinding = bOverrideParentBinding;
+		InputTouchReleasedEvent->InputKeyEvent = IE_Released;
+		InputTouchReleasedEvent->EventSignatureName = TEXT("InputTouchHandlerDynamicSignature__DelegateSignature");
+		InputTouchReleasedEvent->EventSignatureClass = UInputComponent::StaticClass();
+		InputTouchReleasedEvent->bInternalEvent = true;
+		InputTouchReleasedEvent->AllocateDefaultPins();
+
+		// Create a temporary variable to copy location in to
+		static UScriptStruct* VectorStruct = FindObjectChecked<UScriptStruct>(UObject::StaticClass(), TEXT("Vector"));
+
+		UK2Node_TemporaryVariable* TouchLocationVar = CompilerContext.SpawnIntermediateNode<UK2Node_TemporaryVariable>(this, SourceGraph);
+		TouchLocationVar->VariableType.PinCategory = Schema->PC_Struct;
+		TouchLocationVar->VariableType.PinSubCategoryObject = VectorStruct;
+		TouchLocationVar->AllocateDefaultPins();
+
+		// Create assignment nodes to assign the location
+		UK2Node_AssignmentStatement* TouchPressedLocationInitialize = CompilerContext.SpawnIntermediateNode<UK2Node_AssignmentStatement>(this, SourceGraph);
+		TouchPressedLocationInitialize->AllocateDefaultPins();
+		Schema->TryCreateConnection(TouchLocationVar->GetVariablePin(), TouchPressedLocationInitialize->GetVariablePin());
+		Schema->TryCreateConnection(TouchPressedLocationInitialize->GetValuePin(), InputTouchPressedEvent->FindPinChecked(TEXT("Location")));
+
+		UK2Node_AssignmentStatement* TouchReleasedLocationInitialize = CompilerContext.SpawnIntermediateNode<UK2Node_AssignmentStatement>(this, SourceGraph);
+		TouchReleasedLocationInitialize->AllocateDefaultPins();
+		Schema->TryCreateConnection(TouchLocationVar->GetVariablePin(), TouchReleasedLocationInitialize->GetVariablePin());
+		Schema->TryCreateConnection(TouchReleasedLocationInitialize->GetValuePin(), InputTouchReleasedEvent->FindPinChecked(TEXT("Location")));
+
+		// Connect the events to the assign location nodes
+		Schema->TryCreateConnection(Schema->FindExecutionPin(*InputTouchPressedEvent, EGPD_Output), TouchPressedLocationInitialize->GetExecPin());
+		Schema->TryCreateConnection(Schema->FindExecutionPin(*InputTouchReleasedEvent, EGPD_Output), TouchReleasedLocationInitialize->GetExecPin());
+
+		// Create a temporary variable to copy finger index in to
+		UK2Node_TemporaryVariable* TouchFingerVar = CompilerContext.SpawnIntermediateNode<UK2Node_TemporaryVariable>(this, SourceGraph);
+		TouchFingerVar->VariableType.PinCategory = Schema->PC_Byte;
+		TouchFingerVar->VariableType.PinSubCategoryObject = UK2Node_InputTouch::GetTouchIndexEnum();
+		TouchFingerVar->AllocateDefaultPins();
+
+		// Create assignment nodes to assign the finger index
+		UK2Node_AssignmentStatement* TouchPressedFingerInitialize = CompilerContext.SpawnIntermediateNode<UK2Node_AssignmentStatement>(this, SourceGraph);
+		TouchPressedFingerInitialize->AllocateDefaultPins();
+		Schema->TryCreateConnection(TouchFingerVar->GetVariablePin(), TouchPressedFingerInitialize->GetVariablePin());
+		Schema->TryCreateConnection(TouchPressedFingerInitialize->GetValuePin(), InputTouchPressedEvent->FindPinChecked(TEXT("FingerIndex")));
+
+		UK2Node_AssignmentStatement* TouchReleasedFingerInitialize = CompilerContext.SpawnIntermediateNode<UK2Node_AssignmentStatement>(this, SourceGraph);
+		TouchReleasedFingerInitialize->AllocateDefaultPins();
+		Schema->TryCreateConnection(TouchFingerVar->GetVariablePin(), TouchReleasedFingerInitialize->GetVariablePin());
+		Schema->TryCreateConnection(TouchReleasedFingerInitialize->GetValuePin(), InputTouchReleasedEvent->FindPinChecked(TEXT("FingerIndex")));
+
+		// Connect the assign location to the assign finger index nodes
+		Schema->TryCreateConnection(TouchPressedLocationInitialize->GetThenPin(), TouchPressedFingerInitialize->GetExecPin());
+		Schema->TryCreateConnection(TouchReleasedLocationInitialize->GetThenPin(), TouchReleasedFingerInitialize->GetExecPin());
+
+		// Move the original event connections to the then pin of the finger index assign
+		CompilerContext.MovePinLinksToIntermediate(*GetPressedPin(), *TouchPressedFingerInitialize->GetThenPin());
+		CompilerContext.MovePinLinksToIntermediate(*GetReleasedPin(), *TouchReleasedFingerInitialize->GetThenPin());
+
+		// Move the original event variable connections to the intermediate nodes
+		CompilerContext.MovePinLinksToIntermediate(*GetLocationPin(), *TouchLocationVar->GetVariablePin());
+		CompilerContext.MovePinLinksToIntermediate(*GetFingerIndexPin(), *TouchFingerVar->GetVariablePin());
+	}
+	else
+	{
+		UEdGraphPin* InputTouchPin; 
+		EInputEvent InputEvent;
+		if (InputTouchPressedPin->LinkedTo.Num() > 0)
 		{
-			// Create the input touch events
-			UK2Node_InputTouchEvent* InputTouchPressedEvent = CompilerContext.SpawnIntermediateNode<UK2Node_InputTouchEvent>(this, SourceGraph);
-			InputTouchPressedEvent->CustomFunctionName = FName( *FString::Printf(TEXT("InpTchEvt_%s"), *InputTouchPressedEvent->GetName()));
-			InputTouchPressedEvent->bConsumeInput = bConsumeInput;
-			InputTouchPressedEvent->bExecuteWhenPaused = bExecuteWhenPaused;
-			InputTouchPressedEvent->bOverrideParentBinding = bOverrideParentBinding;
-			InputTouchPressedEvent->InputKeyEvent = IE_Pressed;
-			InputTouchPressedEvent->EventSignatureName = TEXT("InputTouchHandlerDynamicSignature__DelegateSignature");
-			InputTouchPressedEvent->EventSignatureClass = UInputComponent::StaticClass();
-			InputTouchPressedEvent->bInternalEvent = true;
-			InputTouchPressedEvent->AllocateDefaultPins();
-
-			UK2Node_InputTouchEvent* InputTouchReleasedEvent = CompilerContext.SpawnIntermediateNode<UK2Node_InputTouchEvent>(this, SourceGraph);
-			InputTouchReleasedEvent->CustomFunctionName = FName( *FString::Printf(TEXT("InpTchEvt_%s"),  *InputTouchReleasedEvent->GetName()));
-			InputTouchReleasedEvent->bConsumeInput = bConsumeInput;
-			InputTouchReleasedEvent->bExecuteWhenPaused = bExecuteWhenPaused;
-			InputTouchReleasedEvent->bOverrideParentBinding = bOverrideParentBinding;
-			InputTouchReleasedEvent->InputKeyEvent = IE_Released;
-			InputTouchReleasedEvent->EventSignatureName = TEXT("InputTouchHandlerDynamicSignature__DelegateSignature");
-			InputTouchReleasedEvent->EventSignatureClass = UInputComponent::StaticClass();
-			InputTouchReleasedEvent->bInternalEvent = true;
-			InputTouchReleasedEvent->AllocateDefaultPins();
-
-			// Create a temporary variable to copy location in to
-			static UScriptStruct* VectorStruct = FindObjectChecked<UScriptStruct>(UObject::StaticClass(), TEXT("Vector"));
-
-			UK2Node_TemporaryVariable* TouchLocationVar = CompilerContext.SpawnIntermediateNode<UK2Node_TemporaryVariable>(this, SourceGraph);
-			TouchLocationVar->VariableType.PinCategory = Schema->PC_Struct;
-			TouchLocationVar->VariableType.PinSubCategoryObject = VectorStruct;
-			TouchLocationVar->AllocateDefaultPins();
-
-			// Create assignment nodes to assign the location
-			UK2Node_AssignmentStatement* TouchPressedLocationInitialize = CompilerContext.SpawnIntermediateNode<UK2Node_AssignmentStatement>(this, SourceGraph);
-			TouchPressedLocationInitialize->AllocateDefaultPins();
-			Schema->TryCreateConnection(TouchLocationVar->GetVariablePin(), TouchPressedLocationInitialize->GetVariablePin());
-			Schema->TryCreateConnection(TouchPressedLocationInitialize->GetValuePin(), InputTouchPressedEvent->FindPinChecked(TEXT("Location")));
-
-			UK2Node_AssignmentStatement* TouchReleasedLocationInitialize = CompilerContext.SpawnIntermediateNode<UK2Node_AssignmentStatement>(this, SourceGraph);
-			TouchReleasedLocationInitialize->AllocateDefaultPins();
-			Schema->TryCreateConnection(TouchLocationVar->GetVariablePin(), TouchReleasedLocationInitialize->GetVariablePin());
-			Schema->TryCreateConnection(TouchReleasedLocationInitialize->GetValuePin(), InputTouchReleasedEvent->FindPinChecked(TEXT("Location")));
-
-			// Connect the events to the assign location nodes
-			Schema->TryCreateConnection(Schema->FindExecutionPin(*InputTouchPressedEvent, EGPD_Output), TouchPressedLocationInitialize->GetExecPin());
-			Schema->TryCreateConnection(Schema->FindExecutionPin(*InputTouchReleasedEvent, EGPD_Output), TouchReleasedLocationInitialize->GetExecPin());
-
-			// Create a temporary variable to copy finger index in to
-			UK2Node_TemporaryVariable* TouchFingerVar = CompilerContext.SpawnIntermediateNode<UK2Node_TemporaryVariable>(this, SourceGraph);
-			TouchFingerVar->VariableType.PinCategory = Schema->PC_Byte;
-			TouchFingerVar->VariableType.PinSubCategoryObject = UK2Node_InputTouch::GetTouchIndexEnum();
-			TouchFingerVar->AllocateDefaultPins();
-
-			// Create assignment nodes to assign the finger index
-			UK2Node_AssignmentStatement* TouchPressedFingerInitialize = CompilerContext.SpawnIntermediateNode<UK2Node_AssignmentStatement>(this, SourceGraph);
-			TouchPressedFingerInitialize->AllocateDefaultPins();
-			Schema->TryCreateConnection(TouchFingerVar->GetVariablePin(), TouchPressedFingerInitialize->GetVariablePin());
-			Schema->TryCreateConnection(TouchPressedFingerInitialize->GetValuePin(), InputTouchPressedEvent->FindPinChecked(TEXT("FingerIndex")));
-
-			UK2Node_AssignmentStatement* TouchReleasedFingerInitialize = CompilerContext.SpawnIntermediateNode<UK2Node_AssignmentStatement>(this, SourceGraph);
-			TouchReleasedFingerInitialize->AllocateDefaultPins();
-			Schema->TryCreateConnection(TouchFingerVar->GetVariablePin(), TouchReleasedFingerInitialize->GetVariablePin());
-			Schema->TryCreateConnection(TouchReleasedFingerInitialize->GetValuePin(), InputTouchReleasedEvent->FindPinChecked(TEXT("FingerIndex")));
-
-			// Connect the assign location to the assign finger index nodes
-			Schema->TryCreateConnection(TouchPressedLocationInitialize->GetThenPin(), TouchPressedFingerInitialize->GetExecPin());
-			Schema->TryCreateConnection(TouchReleasedLocationInitialize->GetThenPin(), TouchReleasedFingerInitialize->GetExecPin());
-
-			// Move the original event connections to the then pin of the finger index assign
-			CompilerContext.MovePinLinksToIntermediate(*GetPressedPin(), *TouchPressedFingerInitialize->GetThenPin());
-			CompilerContext.MovePinLinksToIntermediate(*GetReleasedPin(), *TouchReleasedFingerInitialize->GetThenPin());
-
-			// Move the original event variable connections to the intermediate nodes
-			CompilerContext.MovePinLinksToIntermediate(*GetLocationPin(), *TouchLocationVar->GetVariablePin());
-			CompilerContext.MovePinLinksToIntermediate(*GetFingerIndexPin(), *TouchFingerVar->GetVariablePin());
+			InputTouchPin = InputTouchPressedPin;
+			InputEvent = IE_Pressed;
 		}
 		else
 		{
-			UEdGraphPin* InputTouchPin; 
-			EInputEvent InputEvent;
-			if (InputTouchPressedPin->LinkedTo.Num() > 0)
-			{
-				InputTouchPin = InputTouchPressedPin;
-				InputEvent = IE_Pressed;
-			}
-			else
-			{
-				InputTouchPin = InputTouchReleasedPin;
-				InputEvent = IE_Released;
-			}
-
-			if (InputTouchPin->LinkedTo.Num() > 0)
-			{
-				UK2Node_InputTouchEvent* InputTouchEvent = CompilerContext.SpawnIntermediateNode<UK2Node_InputTouchEvent>(this, SourceGraph);
-				InputTouchEvent->CustomFunctionName = FName( *FString::Printf(TEXT("InpTchEvt_%s"), *InputTouchEvent->GetName()));
-				InputTouchEvent->InputKeyEvent = InputEvent;
-				InputTouchEvent->bConsumeInput = bConsumeInput;
-				InputTouchEvent->bExecuteWhenPaused = bExecuteWhenPaused;
-				InputTouchEvent->bOverrideParentBinding = bOverrideParentBinding;
-				InputTouchEvent->EventSignatureName = TEXT("InputTouchHandlerDynamicSignature__DelegateSignature");
-				InputTouchEvent->EventSignatureClass = UInputComponent::StaticClass();
-				InputTouchEvent->bInternalEvent = true;
-				InputTouchEvent->AllocateDefaultPins();
-
-				CompilerContext.MovePinLinksToIntermediate(*InputTouchPin, *Schema->FindExecutionPin(*InputTouchEvent, EGPD_Output));
-				CompilerContext.MovePinLinksToIntermediate(*GetLocationPin(), *InputTouchEvent->FindPin(TEXT("Location")));
-				CompilerContext.MovePinLinksToIntermediate(*GetFingerIndexPin(), *InputTouchEvent->FindPin(TEXT("FingerIndex")));
-			}
+			InputTouchPin = InputTouchReleasedPin;
+			InputEvent = IE_Released;
+		}
+
+		if (InputTouchPin->LinkedTo.Num() > 0)
+		{
+			UK2Node_InputTouchEvent* InputTouchEvent = CompilerContext.SpawnIntermediateNode<UK2Node_InputTouchEvent>(this, SourceGraph);
+			InputTouchEvent->CustomFunctionName = FName( *FString::Printf(TEXT("InpTchEvt_%s"), *InputTouchEvent->GetName()));
+			InputTouchEvent->InputKeyEvent = InputEvent;
+			InputTouchEvent->bConsumeInput = bConsumeInput;
+			InputTouchEvent->bExecuteWhenPaused = bExecuteWhenPaused;
+			InputTouchEvent->bOverrideParentBinding = bOverrideParentBinding;
+			InputTouchEvent->EventSignatureName = TEXT("InputTouchHandlerDynamicSignature__DelegateSignature");
+			InputTouchEvent->EventSignatureClass = UInputComponent::StaticClass();
+			InputTouchEvent->bInternalEvent = true;
+			InputTouchEvent->AllocateDefaultPins();
+
+			CompilerContext.MovePinLinksToIntermediate(*InputTouchPin, *Schema->FindExecutionPin(*InputTouchEvent, EGPD_Output));
+			CompilerContext.MovePinLinksToIntermediate(*GetLocationPin(), *InputTouchEvent->FindPin(TEXT("Location")));
+			CompilerContext.MovePinLinksToIntermediate(*GetFingerIndexPin(), *InputTouchEvent->FindPin(TEXT("FingerIndex")));
 		}
 	}
 }
\ No newline at end of file
diff --git a/Engine/Source/Editor/BlueprintGraph/Private/K2Node_MCDelegate.cpp b/Engine/Source/Editor/BlueprintGraph/Private/K2Node_MCDelegate.cpp
--- a/Engine/Source/Editor/BlueprintGraph/Private/K2Node_MCDelegate.cpp
+++ b/Engine/Source/Editor/BlueprintGraph/Private/K2Node_MCDelegate.cpp
@@ -159,7 +159,7 @@ void UK2Node_BaseMCDelegate::ExpandNode(class FKismetCompilerContext& CompilerCo
 	Super::ExpandNode(CompilerContext, SourceGraph);
 
 	const bool bAllowMultipleSelfs = AllowMultipleSelfs(true);
-	if(bAllowMultipleSelfs && CompilerContext.bIsFullCompile)
+	if(bAllowMultipleSelfs)
 	{
 		const UEdGraphSchema_K2* Schema = CompilerContext.GetSchema();
 		check(Schema);
diff --git a/Engine/Source/Editor/BlueprintGraph/Private/K2Node_Message.cpp b/Engine/Source/Editor/BlueprintGraph/Private/K2Node_Message.cpp
--- a/Engine/Source/Editor/BlueprintGraph/Private/K2Node_Message.cpp
+++ b/Engine/Source/Editor/BlueprintGraph/Private/K2Node_Message.cpp
@@ -72,164 +72,161 @@ void UK2Node_Message::ExpandNode(class FKismetCompilerContext& CompilerContext,
 {
 	Super::ExpandNode(CompilerContext, SourceGraph);
 
-	if (CompilerContext.bIsFullCompile)
-	{
-		const UEdGraphSchema_K2* Schema = CompilerContext.GetSchema();
+	const UEdGraphSchema_K2* Schema = CompilerContext.GetSchema();
 
-		UEdGraphPin* ExecPin = Schema->FindExecutionPin(*this, EGPD_Input);
-		const bool bExecPinConnected = ExecPin && (ExecPin->LinkedTo.Num() > 0);
-		UEdGraphPin* ThenPin = Schema->FindExecutionPin(*this, EGPD_Output);
-		const bool bThenPinConnected = ThenPin && (ThenPin->LinkedTo.Num() > 0);
+	UEdGraphPin* ExecPin = Schema->FindExecutionPin(*this, EGPD_Input);
+	const bool bExecPinConnected = ExecPin && (ExecPin->LinkedTo.Num() > 0);
+	UEdGraphPin* ThenPin = Schema->FindExecutionPin(*this, EGPD_Output);
+	const bool bThenPinConnected = ThenPin && (ThenPin->LinkedTo.Num() > 0);
 
-		// Skip ourselves if our exec isn't wired up
-		if( bExecPinConnected )
+	// Skip ourselves if our exec isn't wired up
+	if( bExecPinConnected )
+	{
+		// Make sure our interface is valid
+		if (FunctionReference.GetMemberParentClass(this) == NULL)
 		{
-			// Make sure our interface is valid
-			if (FunctionReference.GetMemberParentClass(this) == NULL)
-			{
-				CompilerContext.MessageLog.Error(*FString::Printf(*LOCTEXT("MessageNodeInvalid_Error", "Message node @@ has an invalid interface.").ToString()), this);
-				return;
-			}
+			CompilerContext.MessageLog.Error(*FString::Printf(*LOCTEXT("MessageNodeInvalid_Error", "Message node @@ has an invalid interface.").ToString()), this);
+			return;
+		}
 
-			UFunction* MessageNodeFunction = GetTargetFunction();
-			if (MessageNodeFunction == NULL)
-			{
-				//@TODO: Why do this here in teh compiler, it's already done on AllocateDefaultPins() during on-load node reconstruction
-				MessageNodeFunction = Cast<UFunction>(UK2Node::FindRemappedField(FunctionReference.GetMemberParentClass(this), FunctionReference.GetMemberName()));
-			}
+		UFunction* MessageNodeFunction = GetTargetFunction();
+		if (MessageNodeFunction == NULL)
+		{
+			//@TODO: Why do this here in teh compiler, it's already done on AllocateDefaultPins() during on-load node reconstruction
+			MessageNodeFunction = Cast<UFunction>(UK2Node::FindRemappedField(FunctionReference.GetMemberParentClass(this), FunctionReference.GetMemberName()));
+		}
 
-			if (MessageNodeFunction == NULL)
-			{
-				CompilerContext.MessageLog.Error(*FString::Printf(*LOCTEXT("MessageNodeInvalidFunction_Error", "Unable to find function with name %s for Message node @@.").ToString(), *(FunctionReference.GetMemberName().ToString())), this);
-				return;
-			}
+		if (MessageNodeFunction == NULL)
+		{
+			CompilerContext.MessageLog.Error(*FString::Printf(*LOCTEXT("MessageNodeInvalidFunction_Error", "Unable to find function with name %s for Message node @@.").ToString(), *(FunctionReference.GetMemberName().ToString())), this);
+			return;
+		}
 
-			// Check to make sure we have a target
-			UEdGraphPin* MessageSelfPin = Schema->FindSelfPin(*this, EGPD_Input);
-			if( !MessageSelfPin || MessageSelfPin->LinkedTo.Num() == 0 )
-			{
-				CompilerContext.MessageLog.Error(*FString::Printf(*LOCTEXT("MessageNodeSelfPin_Error", "Message node @@ must have a self pin connection.").ToString()), this);
-				return;
-			}
+		// Check to make sure we have a target
+		UEdGraphPin* MessageSelfPin = Schema->FindSelfPin(*this, EGPD_Input);
+		if( !MessageSelfPin || MessageSelfPin->LinkedTo.Num() == 0 )
+		{
+			CompilerContext.MessageLog.Error(*FString::Printf(*LOCTEXT("MessageNodeSelfPin_Error", "Message node @@ must have a self pin connection.").ToString()), this);
+			return;
+		}
 
-			// First, create an internal cast-to-interface node
-			UK2Node_DynamicCast* CastToInterfaceNode = CompilerContext.SpawnIntermediateNode<UK2Node_DynamicCast>(this, SourceGraph);
-			CastToInterfaceNode->TargetType = MessageNodeFunction->GetOuterUClass();
-			CastToInterfaceNode->AllocateDefaultPins();
+		// First, create an internal cast-to-interface node
+		UK2Node_DynamicCast* CastToInterfaceNode = CompilerContext.SpawnIntermediateNode<UK2Node_DynamicCast>(this, SourceGraph);
+		CastToInterfaceNode->TargetType = MessageNodeFunction->GetOuterUClass();
+		CastToInterfaceNode->AllocateDefaultPins();
 
-			UEdGraphPin* CastToInterfaceValidPin = CastToInterfaceNode->GetValidCastPin();
-			UEdGraphPin* CastToInterfaceResultPin = CastToInterfaceNode->GetCastResultPin();
-			if( !CastToInterfaceResultPin )
-			{
-				CompilerContext.MessageLog.Error(*LOCTEXT("InvalidInterfaceClass_Error", "Node @@ has an invalid target interface class").ToString(), this);
-				return;
-			}
+		UEdGraphPin* CastToInterfaceValidPin = CastToInterfaceNode->GetValidCastPin();
+		UEdGraphPin* CastToInterfaceResultPin = CastToInterfaceNode->GetCastResultPin();
+		if( !CastToInterfaceResultPin )
+		{
+			CompilerContext.MessageLog.Error(*LOCTEXT("InvalidInterfaceClass_Error", "Node @@ has an invalid target interface class").ToString(), this);
+			return;
+		}
 
-			CastToInterfaceResultPin->PinType.PinSubCategoryObject = *CastToInterfaceNode->TargetType;
+		CastToInterfaceResultPin->PinType.PinSubCategoryObject = *CastToInterfaceNode->TargetType;
 
-			// Wire up the connections
-			CompilerContext.MovePinLinksToIntermediate(*ExecPin, *CastToInterfaceNode->GetExecPin());
+		// Wire up the connections
+		CompilerContext.MovePinLinksToIntermediate(*ExecPin, *CastToInterfaceNode->GetExecPin());
 
-			UEdGraphPin* CastToInterfaceSourceObjectPin = CastToInterfaceNode->GetCastSourcePin();
-			CompilerContext.MovePinLinksToIntermediate(*MessageSelfPin, *CastToInterfaceSourceObjectPin);
+		UEdGraphPin* CastToInterfaceSourceObjectPin = CastToInterfaceNode->GetCastSourcePin();
+		CompilerContext.MovePinLinksToIntermediate(*MessageSelfPin, *CastToInterfaceSourceObjectPin);
 
-			// Next, create the function call node
-			UK2Node_CallFunction* FunctionCallNode = CompilerContext.SpawnIntermediateNode<UK2Node_CallFunction>(this, SourceGraph);
-			FunctionCallNode->bIsInterfaceCall = true;
-			FunctionCallNode->FunctionReference = FunctionReference;
-			FunctionCallNode->AllocateDefaultPins();
+		// Next, create the function call node
+		UK2Node_CallFunction* FunctionCallNode = CompilerContext.SpawnIntermediateNode<UK2Node_CallFunction>(this, SourceGraph);
+		FunctionCallNode->bIsInterfaceCall = true;
+		FunctionCallNode->FunctionReference = FunctionReference;
+		FunctionCallNode->AllocateDefaultPins();
 
-			// Wire up the connections
-			UEdGraphPin* CallFunctionExecPin = Schema->FindExecutionPin(*FunctionCallNode, EGPD_Input);
-			CastToInterfaceValidPin->MakeLinkTo(CallFunctionExecPin);
+		// Wire up the connections
+		UEdGraphPin* CallFunctionExecPin = Schema->FindExecutionPin(*FunctionCallNode, EGPD_Input);
+		CastToInterfaceValidPin->MakeLinkTo(CallFunctionExecPin);
 
-			// Self pin
-			UEdGraphPin* FunctionCallSelfPin = Schema->FindSelfPin(*FunctionCallNode, EGPD_Input);
-			CastToInterfaceResultPin->MakeLinkTo(FunctionCallSelfPin);
+		// Self pin
+		UEdGraphPin* FunctionCallSelfPin = Schema->FindSelfPin(*FunctionCallNode, EGPD_Input);
+		CastToInterfaceResultPin->MakeLinkTo(FunctionCallSelfPin);
 
-			UEdGraphPin* LastOutCastFaildPin = CastToInterfaceNode->GetInvalidCastPin();
+		UEdGraphPin* LastOutCastFaildPin = CastToInterfaceNode->GetInvalidCastPin();
 
-			UFunction* ArrayClearFunction = UKismetArrayLibrary::StaticClass()->FindFunctionByName(FName(TEXT("Array_Clear")));
-			check(ArrayClearFunction);
+		UFunction* ArrayClearFunction = UKismetArrayLibrary::StaticClass()->FindFunctionByName(FName(TEXT("Array_Clear")));
+		check(ArrayClearFunction);
 
-			// Variable pins - Try to associate variable inputs to the message node with the variable inputs and outputs to the call function node
-			for( int32 i = 0; i < Pins.Num(); i++ )
+		// Variable pins - Try to associate variable inputs to the message node with the variable inputs and outputs to the call function node
+		for( int32 i = 0; i < Pins.Num(); i++ )
+		{
+			UEdGraphPin* CurrentPin = Pins[i];
+			if( CurrentPin && (CurrentPin->PinType.PinCategory != Schema->PC_Exec) && (CurrentPin->PinName != Schema->PN_Self) )
 			{
-				UEdGraphPin* CurrentPin = Pins[i];
-				if( CurrentPin && (CurrentPin->PinType.PinCategory != Schema->PC_Exec) && (CurrentPin->PinName != Schema->PN_Self) )
+				// Try to find a match for the pin on the function call node
+				UEdGraphPin* FunctionCallPin = FunctionCallNode->FindPin(CurrentPin->PinName);
+				if( FunctionCallPin )
 				{
-					// Try to find a match for the pin on the function call node
-					UEdGraphPin* FunctionCallPin = FunctionCallNode->FindPin(CurrentPin->PinName);
-					if( FunctionCallPin )
-					{
-						// Move pin links if the pin is connected...
-						CompilerContext.MovePinLinksToIntermediate(*CurrentPin, *FunctionCallPin);
+					// Move pin links if the pin is connected...
+					CompilerContext.MovePinLinksToIntermediate(*CurrentPin, *FunctionCallPin);
 
-						// when cast fails all return values must be cleared.
-						if (EEdGraphPinDirection::EGPD_Output == CurrentPin->Direction)
-						{
-							UK2Node* DefaultValueNode = NULL;
-							UEdGraphPin* DefaultValueThenPin = NULL;
-							if (CurrentPin->PinType.bIsArray)
-							{
-								UK2Node_CallArrayFunction* ClearArray = CompilerContext.SpawnIntermediateNode<UK2Node_CallArrayFunction>(this, SourceGraph);
-								DefaultValueNode = ClearArray;
-								ClearArray->SetFromFunction(ArrayClearFunction);
-								ClearArray->AllocateDefaultPins();
-
-								UEdGraphPin* ArrayPin = ClearArray->GetTargetArrayPin();
-								check(ArrayPin);
-								Schema->TryCreateConnection(ArrayPin, FunctionCallPin);
-								ClearArray->PinConnectionListChanged(ArrayPin);
-
-								DefaultValueThenPin = ClearArray->GetThenPin();
-							} 
-							else
-							{
-								UK2Node_AssignmentStatement* AssignDefaultValue = CompilerContext.SpawnIntermediateNode<UK2Node_AssignmentStatement>(this, SourceGraph);
-								DefaultValueNode = AssignDefaultValue;
-								AssignDefaultValue->AllocateDefaultPins();
-
-								Schema->TryCreateConnection(AssignDefaultValue->GetVariablePin(), FunctionCallPin);
-								AssignDefaultValue->PinConnectionListChanged(AssignDefaultValue->GetVariablePin());
-								Schema->SetPinDefaultValueBasedOnType(AssignDefaultValue->GetValuePin());
-
-								DefaultValueThenPin = AssignDefaultValue->GetThenPin();
-							}
-							UEdGraphPin* DefaultValueExecPin = DefaultValueNode->GetExecPin();
-							check(DefaultValueExecPin);
-							Schema->TryCreateConnection(DefaultValueExecPin, LastOutCastFaildPin);
-							LastOutCastFaildPin = DefaultValueThenPin;
-							check(LastOutCastFaildPin);
-						}
-					}
-					else
+					// when cast fails all return values must be cleared.
+					if (EEdGraphPinDirection::EGPD_Output == CurrentPin->Direction)
 					{
-						UE_LOG(LogK2Compiler, Log, TEXT("%s"), *LOCTEXT("NoPinConnectionFound_Error", "Unable to find connection for pin!  Check AllocateDefaultPins() for consistency!").ToString());
-					}
-				}
-			}
+						UK2Node* DefaultValueNode = NULL;
+						UEdGraphPin* DefaultValueThenPin = NULL;
+						if (CurrentPin->PinType.bIsArray)
+						{
+							UK2Node_CallArrayFunction* ClearArray = CompilerContext.SpawnIntermediateNode<UK2Node_CallArrayFunction>(this, SourceGraph);
+							DefaultValueNode = ClearArray;
+							ClearArray->SetFromFunction(ArrayClearFunction);
+							ClearArray->AllocateDefaultPins();
+
+							UEdGraphPin* ArrayPin = ClearArray->GetTargetArrayPin();
+							check(ArrayPin);
+							Schema->TryCreateConnection(ArrayPin, FunctionCallPin);
+							ClearArray->PinConnectionListChanged(ArrayPin);
+
+							DefaultValueThenPin = ClearArray->GetThenPin();
+						} 
+						else
+						{
+							UK2Node_AssignmentStatement* AssignDefaultValue = CompilerContext.SpawnIntermediateNode<UK2Node_AssignmentStatement>(this, SourceGraph);
+							DefaultValueNode = AssignDefaultValue;
+							AssignDefaultValue->AllocateDefaultPins();
 
-			if( bThenPinConnected )
-			{
-				// Failure case for the cast runs straight through to the exit
-				CompilerContext.CopyPinLinksToIntermediate(*ThenPin, *LastOutCastFaildPin);
+							Schema->TryCreateConnection(AssignDefaultValue->GetVariablePin(), FunctionCallPin);
+							AssignDefaultValue->PinConnectionListChanged(AssignDefaultValue->GetVariablePin());
+							Schema->SetPinDefaultValueBasedOnType(AssignDefaultValue->GetValuePin());
 
-				// Copy all links from the invalid cast case above to the call function node
-				if (UEdGraphPin* CallFunctionThenPin = Schema->FindExecutionPin(*FunctionCallNode, EGPD_Output))
-				{
-					CompilerContext.MovePinLinksToIntermediate(*ThenPin, *CallFunctionThenPin);
+							DefaultValueThenPin = AssignDefaultValue->GetThenPin();
+						}
+						UEdGraphPin* DefaultValueExecPin = DefaultValueNode->GetExecPin();
+						check(DefaultValueExecPin);
+						Schema->TryCreateConnection(DefaultValueExecPin, LastOutCastFaildPin);
+						LastOutCastFaildPin = DefaultValueThenPin;
+						check(LastOutCastFaildPin);
+					}
 				}
 				else
 				{
-					CompilerContext.MessageLog.Warning(*LOCTEXT("ExpectedThenPinOnInterfaceCall", "Expected a then pin on the call in @@.  Does the interface method still exist?").ToString(), FunctionCallNode);
+					UE_LOG(LogK2Compiler, Log, TEXT("%s"), *LOCTEXT("NoPinConnectionFound_Error", "Unable to find connection for pin!  Check AllocateDefaultPins() for consistency!").ToString());
 				}
 			}
 		}
 
-		// Break all connections to the original node, so it will be pruned
-		BreakAllNodeLinks();
+		if( bThenPinConnected )
+		{
+			// Failure case for the cast runs straight through to the exit
+			CompilerContext.CopyPinLinksToIntermediate(*ThenPin, *LastOutCastFaildPin);
+
+			// Copy all links from the invalid cast case above to the call function node
+			if (UEdGraphPin* CallFunctionThenPin = Schema->FindExecutionPin(*FunctionCallNode, EGPD_Output))
+			{
+				CompilerContext.MovePinLinksToIntermediate(*ThenPin, *CallFunctionThenPin);
+			}
+			else
+			{
+				CompilerContext.MessageLog.Warning(*LOCTEXT("ExpectedThenPinOnInterfaceCall", "Expected a then pin on the call in @@.  Does the interface method still exist?").ToString(), FunctionCallNode);
+			}
+		}
 	}
+
+	// Break all connections to the original node, so it will be pruned
+	BreakAllNodeLinks();
 }
 
 #undef LOCTEXT_NAMESPACE
\ No newline at end of file
diff --git a/Engine/Source/Editor/BlueprintGraph/Private/K2Node_MultiGate.cpp b/Engine/Source/Editor/BlueprintGraph/Private/K2Node_MultiGate.cpp
--- a/Engine/Source/Editor/BlueprintGraph/Private/K2Node_MultiGate.cpp
+++ b/Engine/Source/Editor/BlueprintGraph/Private/K2Node_MultiGate.cpp
@@ -607,48 +607,45 @@ void UK2Node_MultiGate::ExpandNode(class FKismetCompilerContext& CompilerContext
 {
 	Super::ExpandNode(CompilerContext, SourceGraph);
 
-	if (CompilerContext.bIsFullCompile)
+	/////////////////////////////
+	// Handle the "Reset"
+	/////////////////////////////
+
+	// Redirect the reset pin if linked to
+	UEdGraphPin* ResetPin = GetResetPin();
+	if (ResetPin->LinkedTo.Num() > 0)
 	{
+		const UEdGraphSchema_K2* Schema = CompilerContext.GetSchema();
+
 		/////////////////////////////
-		// Handle the "Reset"
+		// Temporary Variable node
 		/////////////////////////////
 
-		// Redirect the reset pin if linked to
-		UEdGraphPin* ResetPin = GetResetPin();
-		if (ResetPin->LinkedTo.Num() > 0)
-		{
-			const UEdGraphSchema_K2* Schema = CompilerContext.GetSchema();
-
-			/////////////////////////////
-			// Temporary Variable node
-			/////////////////////////////
-
-			// Create the node
-			UK2Node_TemporaryVariable* TempVarNode = SourceGraph->CreateBlankNode<UK2Node_TemporaryVariable>();
-			TempVarNode->VariableType.PinCategory = Schema->PC_Int;
-			TempVarNode->AllocateDefaultPins();
-			CompilerContext.MessageLog.NotifyIntermediateObjectCreation(TempVarNode, this);
-			// Give a reference of the variable node to the multi gate node
-			DataNode = TempVarNode;
-
-			/////////////////////////////
-			// Assignment node
-			/////////////////////////////
-
-			// Create the node
-			UK2Node_AssignmentStatement* AssignmentNode = SourceGraph->CreateBlankNode<UK2Node_AssignmentStatement>();
-			AssignmentNode->AllocateDefaultPins();
-			CompilerContext.MessageLog.NotifyIntermediateObjectCreation(AssignmentNode, this);
-
-			// Coerce the wildcards pin types (set the default of the value to 0)
-			AssignmentNode->GetVariablePin()->PinType = TempVarNode->GetVariablePin()->PinType;
-			AssignmentNode->GetVariablePin()->MakeLinkTo(TempVarNode->GetVariablePin());
-			AssignmentNode->GetValuePin()->PinType = TempVarNode->GetVariablePin()->PinType;
-			AssignmentNode->GetValuePin()->DefaultValue = TEXT("0");
-
-			// Move the "Reset" link to the Assignment node
-			CompilerContext.MovePinLinksToIntermediate(*ResetPin, *AssignmentNode->GetExecPin());
-		}
+		// Create the node
+		UK2Node_TemporaryVariable* TempVarNode = SourceGraph->CreateBlankNode<UK2Node_TemporaryVariable>();
+		TempVarNode->VariableType.PinCategory = Schema->PC_Int;
+		TempVarNode->AllocateDefaultPins();
+		CompilerContext.MessageLog.NotifyIntermediateObjectCreation(TempVarNode, this);
+		// Give a reference of the variable node to the multi gate node
+		DataNode = TempVarNode;
+
+		/////////////////////////////
+		// Assignment node
+		/////////////////////////////
+
+		// Create the node
+		UK2Node_AssignmentStatement* AssignmentNode = SourceGraph->CreateBlankNode<UK2Node_AssignmentStatement>();
+		AssignmentNode->AllocateDefaultPins();
+		CompilerContext.MessageLog.NotifyIntermediateObjectCreation(AssignmentNode, this);
+
+		// Coerce the wildcards pin types (set the default of the value to 0)
+		AssignmentNode->GetVariablePin()->PinType = TempVarNode->GetVariablePin()->PinType;
+		AssignmentNode->GetVariablePin()->MakeLinkTo(TempVarNode->GetVariablePin());
+		AssignmentNode->GetValuePin()->PinType = TempVarNode->GetVariablePin()->PinType;
+		AssignmentNode->GetValuePin()->DefaultValue = TEXT("0");
+
+		// Move the "Reset" link to the Assignment node
+		CompilerContext.MovePinLinksToIntermediate(*ResetPin, *AssignmentNode->GetExecPin());
 	}
 }
 
diff --git a/Engine/Source/Editor/BlueprintGraph/Private/K2Node_SpawnActor.cpp b/Engine/Source/Editor/BlueprintGraph/Private/K2Node_SpawnActor.cpp
--- a/Engine/Source/Editor/BlueprintGraph/Private/K2Node_SpawnActor.cpp
+++ b/Engine/Source/Editor/BlueprintGraph/Private/K2Node_SpawnActor.cpp
@@ -264,189 +264,186 @@ void UK2Node_SpawnActor::ExpandNode(class FKismetCompilerContext& CompilerContex
 {
 	Super::ExpandNode(CompilerContext, SourceGraph);
 
-	if (CompilerContext.bIsFullCompile)
+	static FName BeginSpawningBlueprintFuncName = GET_FUNCTION_NAME_CHECKED(UGameplayStatics, BeginSpawningActorFromBlueprint);
+	static FString BlueprintParamName = FString(TEXT("Blueprint"));
+	static FString WorldContextParamName = FString(TEXT("WorldContextObject"));
+
+	static FName FinishSpawningFuncName = GET_FUNCTION_NAME_CHECKED(UGameplayStatics, FinishSpawningActor);
+	static FString ActorParamName = FString(TEXT("Actor"));
+	static FString TransformParamName = FString(TEXT("SpawnTransform"));
+	static FString NoCollisionFailParamName = FString(TEXT("bNoCollisionFail"));
+
+	static FString ObjectParamName = FString(TEXT("Object"));
+	static FString ValueParamName = FString(TEXT("Value"));
+	static FString PropertyNameParamName = FString(TEXT("PropertyName"));
+
+	const UEdGraphSchema_K2* Schema = CompilerContext.GetSchema();
+
+	UEdGraphPin* SpawnNodeExec = GetExecPin();
+	UEdGraphPin* SpawnNodeTransform = GetSpawnTransformPin();
+	UEdGraphPin* SpawnNodeNoCollisionFail = GetNoCollisionFailPin();
+	UEdGraphPin* SpawnWorldContextPin = GetWorldContextPin();
+	UEdGraphPin* SpawnBlueprintPin = GetBlueprintPin();
+	UEdGraphPin* SpawnNodeThen = GetThenPin();
+	UEdGraphPin* SpawnNodeResult = GetResultPin();
+
+	UBlueprint* SpawnBlueprint = NULL;
+	if(SpawnBlueprintPin != NULL)
 	{
-		static FName BeginSpawningBlueprintFuncName = GET_FUNCTION_NAME_CHECKED(UGameplayStatics, BeginSpawningActorFromBlueprint);
-		static FString BlueprintParamName = FString(TEXT("Blueprint"));
-		static FString WorldContextParamName = FString(TEXT("WorldContextObject"));
-
-		static FName FinishSpawningFuncName = GET_FUNCTION_NAME_CHECKED(UGameplayStatics, FinishSpawningActor);
-		static FString ActorParamName = FString(TEXT("Actor"));
-		static FString TransformParamName = FString(TEXT("SpawnTransform"));
-		static FString NoCollisionFailParamName = FString(TEXT("bNoCollisionFail"));
-
-		static FString ObjectParamName = FString(TEXT("Object"));
-		static FString ValueParamName = FString(TEXT("Value"));
-		static FString PropertyNameParamName = FString(TEXT("PropertyName"));
-
-		const UEdGraphSchema_K2* Schema = CompilerContext.GetSchema();
-
-		UEdGraphPin* SpawnNodeExec = GetExecPin();
-		UEdGraphPin* SpawnNodeTransform = GetSpawnTransformPin();
-		UEdGraphPin* SpawnNodeNoCollisionFail = GetNoCollisionFailPin();
-		UEdGraphPin* SpawnWorldContextPin = GetWorldContextPin();
-		UEdGraphPin* SpawnBlueprintPin = GetBlueprintPin();
-		UEdGraphPin* SpawnNodeThen = GetThenPin();
-		UEdGraphPin* SpawnNodeResult = GetResultPin();
-
-		UBlueprint* SpawnBlueprint = NULL;
-		if(SpawnBlueprintPin != NULL)
-		{
-			SpawnBlueprint = Cast<UBlueprint>(SpawnBlueprintPin->DefaultObject);
-		}
+		SpawnBlueprint = Cast<UBlueprint>(SpawnBlueprintPin->DefaultObject);
+	}
 
-		if(0 == SpawnBlueprintPin->LinkedTo.Num())	
+	if(0 == SpawnBlueprintPin->LinkedTo.Num())	
+	{
+		if(NULL == SpawnBlueprint)
 		{
-			if(NULL == SpawnBlueprint)
-			{
-				CompilerContext.MessageLog.Error(*LOCTEXT("SpawnActorNodeMissingBlueprint_Error", "Spawn node @@ must have a blueprint specified.").ToString(), this);
-				// we break exec links so this is the only error we get, don't want the SpawnActor node being considered and giving 'unexpected node' type warnings
-				BreakAllNodeLinks();
-				return;
-			}
+			CompilerContext.MessageLog.Error(*LOCTEXT("SpawnActorNodeMissingBlueprint_Error", "Spawn node @@ must have a blueprint specified.").ToString(), this);
+			// we break exec links so this is the only error we get, don't want the SpawnActor node being considered and giving 'unexpected node' type warnings
+			BreakAllNodeLinks();
+			return;
+		}
 
-			// check if default blueprint is based on Actor
-			const UClass* GeneratedClass = SpawnBlueprint->GeneratedClass;
-			bool bInvalidBase = GeneratedClass && !GeneratedClass->IsChildOf(AActor::StaticClass());
+		// check if default blueprint is based on Actor
+		const UClass* GeneratedClass = SpawnBlueprint->GeneratedClass;
+		bool bInvalidBase = GeneratedClass && !GeneratedClass->IsChildOf(AActor::StaticClass());
 
-			const UClass* SkeletonGeneratedClass = Cast<UClass>(SpawnBlueprint->SkeletonGeneratedClass);
-			bInvalidBase |= SkeletonGeneratedClass && !SkeletonGeneratedClass->IsChildOf(AActor::StaticClass());
+		const UClass* SkeletonGeneratedClass = Cast<UClass>(SpawnBlueprint->SkeletonGeneratedClass);
+		bInvalidBase |= SkeletonGeneratedClass && !SkeletonGeneratedClass->IsChildOf(AActor::StaticClass());
 
-			if(bInvalidBase)
-			{
-				CompilerContext.MessageLog.Error(*LOCTEXT("SpawnActorNodeInvalidBlueprint_Error", "Spawn node @@ must have a blueprint based on Actor specified.").ToString(), this);
-				// we break exec links so this is the only error we get, don't want the SpawnActor node being considered and giving 'unexpected node' type warnings
-				BreakAllNodeLinks();
-				return;
-			}
+		if(bInvalidBase)
+		{
+			CompilerContext.MessageLog.Error(*LOCTEXT("SpawnActorNodeInvalidBlueprint_Error", "Spawn node @@ must have a blueprint based on Actor specified.").ToString(), this);
+			// we break exec links so this is the only error we get, don't want the SpawnActor node being considered and giving 'unexpected node' type warnings
+			BreakAllNodeLinks();
+			return;
 		}
+	}
 
-		//////////////////////////////////////////////////////////////////////////
-		// create 'begin spawn' call node
-		UK2Node_CallFunction* CallBeginSpawnNode = CompilerContext.SpawnIntermediateNode<UK2Node_CallFunction>(this, SourceGraph);
-		CallBeginSpawnNode->FunctionReference.SetExternalMember(BeginSpawningBlueprintFuncName, UGameplayStatics::StaticClass());
-		CallBeginSpawnNode->AllocateDefaultPins();
+	//////////////////////////////////////////////////////////////////////////
+	// create 'begin spawn' call node
+	UK2Node_CallFunction* CallBeginSpawnNode = CompilerContext.SpawnIntermediateNode<UK2Node_CallFunction>(this, SourceGraph);
+	CallBeginSpawnNode->FunctionReference.SetExternalMember(BeginSpawningBlueprintFuncName, UGameplayStatics::StaticClass());
+	CallBeginSpawnNode->AllocateDefaultPins();
 
-		UEdGraphPin* CallBeginExec = CallBeginSpawnNode->GetExecPin();
-		UEdGraphPin* CallBeginWorldContextPin = CallBeginSpawnNode->FindPinChecked(WorldContextParamName);
-		UEdGraphPin* CallBeginBlueprintPin = CallBeginSpawnNode->FindPinChecked(BlueprintParamName);
-		UEdGraphPin* CallBeginTransform = CallBeginSpawnNode->FindPinChecked(TransformParamName);
-		UEdGraphPin* CallBeginNoCollisionFail = CallBeginSpawnNode->FindPinChecked(NoCollisionFailParamName);
-		UEdGraphPin* CallBeginResult = CallBeginSpawnNode->GetReturnValuePin();
+	UEdGraphPin* CallBeginExec = CallBeginSpawnNode->GetExecPin();
+	UEdGraphPin* CallBeginWorldContextPin = CallBeginSpawnNode->FindPinChecked(WorldContextParamName);
+	UEdGraphPin* CallBeginBlueprintPin = CallBeginSpawnNode->FindPinChecked(BlueprintParamName);
+	UEdGraphPin* CallBeginTransform = CallBeginSpawnNode->FindPinChecked(TransformParamName);
+	UEdGraphPin* CallBeginNoCollisionFail = CallBeginSpawnNode->FindPinChecked(NoCollisionFailParamName);
+	UEdGraphPin* CallBeginResult = CallBeginSpawnNode->GetReturnValuePin();
 
-		// Move 'exec' connection from spawn node to 'begin spawn'
-		CompilerContext.MovePinLinksToIntermediate(*SpawnNodeExec, *CallBeginExec);
+	// Move 'exec' connection from spawn node to 'begin spawn'
+	CompilerContext.MovePinLinksToIntermediate(*SpawnNodeExec, *CallBeginExec);
 
-		if(SpawnBlueprintPin->LinkedTo.Num() > 0)
-		{
-			// Copy the 'blueprint' connection from the spawn node to 'begin spawn'
-			CompilerContext.MovePinLinksToIntermediate(*SpawnBlueprintPin, *CallBeginBlueprintPin);
-		}
-		else
-		{
-			// Copy blueprint literal onto begin spawn call 
-			CallBeginBlueprintPin->DefaultObject = SpawnBlueprint;
-		}
+	if(SpawnBlueprintPin->LinkedTo.Num() > 0)
+	{
+		// Copy the 'blueprint' connection from the spawn node to 'begin spawn'
+		CompilerContext.MovePinLinksToIntermediate(*SpawnBlueprintPin, *CallBeginBlueprintPin);
+	}
+	else
+	{
+		// Copy blueprint literal onto begin spawn call 
+		CallBeginBlueprintPin->DefaultObject = SpawnBlueprint;
+	}
 
-		// Copy the world context connection from the spawn node to 'begin spawn' if necessary
-		if (SpawnWorldContextPin)
-		{
-			CompilerContext.MovePinLinksToIntermediate(*SpawnWorldContextPin, *CallBeginWorldContextPin);
-		}
+	// Copy the world context connection from the spawn node to 'begin spawn' if necessary
+	if (SpawnWorldContextPin)
+	{
+		CompilerContext.MovePinLinksToIntermediate(*SpawnWorldContextPin, *CallBeginWorldContextPin);
+	}
 
-		// Copy the 'transform' connection from the spawn node to 'begin spawn'
-		CompilerContext.MovePinLinksToIntermediate(*SpawnNodeTransform, *CallBeginTransform);
+	// Copy the 'transform' connection from the spawn node to 'begin spawn'
+	CompilerContext.MovePinLinksToIntermediate(*SpawnNodeTransform, *CallBeginTransform);
 		
-		// Copy the 'bNoCollisionFail' connection from the spawn node to 'begin spawn'
-		CompilerContext.MovePinLinksToIntermediate(*SpawnNodeNoCollisionFail, *CallBeginNoCollisionFail);
-
-		//////////////////////////////////////////////////////////////////////////
-		// create 'finish spawn' call node
-		UK2Node_CallFunction* CallFinishSpawnNode = CompilerContext.SpawnIntermediateNode<UK2Node_CallFunction>(this, SourceGraph);
-		CallFinishSpawnNode->FunctionReference.SetExternalMember(FinishSpawningFuncName, UGameplayStatics::StaticClass());
-		CallFinishSpawnNode->AllocateDefaultPins();
-
-		UEdGraphPin* CallFinishExec = CallFinishSpawnNode->GetExecPin();
-		UEdGraphPin* CallFinishThen = CallFinishSpawnNode->GetThenPin();
-		UEdGraphPin* CallFinishActor = CallFinishSpawnNode->FindPinChecked(ActorParamName);
-		UEdGraphPin* CallFinishTransform = CallFinishSpawnNode->FindPinChecked(TransformParamName);
-		UEdGraphPin* CallFinishResult = CallFinishSpawnNode->GetReturnValuePin();
-
-		// Move 'then' connection from spawn node to 'finish spawn'
-		CompilerContext.MovePinLinksToIntermediate(*SpawnNodeThen, *CallFinishThen);
+	// Copy the 'bNoCollisionFail' connection from the spawn node to 'begin spawn'
+	CompilerContext.MovePinLinksToIntermediate(*SpawnNodeNoCollisionFail, *CallBeginNoCollisionFail);
+
+	//////////////////////////////////////////////////////////////////////////
+	// create 'finish spawn' call node
+	UK2Node_CallFunction* CallFinishSpawnNode = CompilerContext.SpawnIntermediateNode<UK2Node_CallFunction>(this, SourceGraph);
+	CallFinishSpawnNode->FunctionReference.SetExternalMember(FinishSpawningFuncName, UGameplayStatics::StaticClass());
+	CallFinishSpawnNode->AllocateDefaultPins();
+
+	UEdGraphPin* CallFinishExec = CallFinishSpawnNode->GetExecPin();
+	UEdGraphPin* CallFinishThen = CallFinishSpawnNode->GetThenPin();
+	UEdGraphPin* CallFinishActor = CallFinishSpawnNode->FindPinChecked(ActorParamName);
+	UEdGraphPin* CallFinishTransform = CallFinishSpawnNode->FindPinChecked(TransformParamName);
+	UEdGraphPin* CallFinishResult = CallFinishSpawnNode->GetReturnValuePin();
+
+	// Move 'then' connection from spawn node to 'finish spawn'
+	CompilerContext.MovePinLinksToIntermediate(*SpawnNodeThen, *CallFinishThen);
 			
-		// Copy transform connection
-		CompilerContext.CopyPinLinksToIntermediate(*CallBeginTransform, *CallFinishTransform);
+	// Copy transform connection
+	CompilerContext.CopyPinLinksToIntermediate(*CallBeginTransform, *CallFinishTransform);
 		
-		// Connect output actor from 'begin' to 'finish'
-		CallBeginResult->MakeLinkTo(CallFinishActor);
+	// Connect output actor from 'begin' to 'finish'
+	CallBeginResult->MakeLinkTo(CallFinishActor);
 
-		// Move result connection from spawn node to 'finish spawn'
-		CallFinishResult->PinType = SpawnNodeResult->PinType; // Copy type so it uses the right actor subclass
-		CompilerContext.MovePinLinksToIntermediate(*SpawnNodeResult, *CallFinishResult);
+	// Move result connection from spawn node to 'finish spawn'
+	CallFinishResult->PinType = SpawnNodeResult->PinType; // Copy type so it uses the right actor subclass
+	CompilerContext.MovePinLinksToIntermediate(*SpawnNodeResult, *CallFinishResult);
 
-		//////////////////////////////////////////////////////////////////////////
-		// create 'set var' nodes
+	//////////////////////////////////////////////////////////////////////////
+	// create 'set var' nodes
 
-		// Get 'result' pin from 'begin spawn', this is the actual actor we want to set properties on
-		UK2Node_CallFunction* LastNode = CallBeginSpawnNode;
+	// Get 'result' pin from 'begin spawn', this is the actual actor we want to set properties on
+	UK2Node_CallFunction* LastNode = CallBeginSpawnNode;
 
-		// Create 'set var by name' nodes and hook them up
-		for(int32 PinIdx=0; PinIdx < Pins.Num(); PinIdx++)
+	// Create 'set var by name' nodes and hook them up
+	for(int32 PinIdx=0; PinIdx < Pins.Num(); PinIdx++)
+	{
+		// Only create 'set param by name' node if this pin is linked to something
+		UEdGraphPin* SpawnVarPin = Pins[PinIdx];
+		if(SpawnVarPin->LinkedTo.Num() > 0)
 		{
-			// Only create 'set param by name' node if this pin is linked to something
-			UEdGraphPin* SpawnVarPin = Pins[PinIdx];
-			if(SpawnVarPin->LinkedTo.Num() > 0)
+			UFunction* SetByNameFunction = Schema->FindSetVariableByNameFunction(SpawnVarPin->PinType);
+			if(SetByNameFunction)
 			{
-				UFunction* SetByNameFunction = Schema->FindSetVariableByNameFunction(SpawnVarPin->PinType);
-				if(SetByNameFunction)
+				UK2Node_CallFunction* SetVarNode = NULL;
+				if(SpawnVarPin->PinType.bIsArray)
 				{
-					UK2Node_CallFunction* SetVarNode = NULL;
-					if(SpawnVarPin->PinType.bIsArray)
-					{
-						SetVarNode = CompilerContext.SpawnIntermediateNode<UK2Node_CallArrayFunction>(this, SourceGraph);
-					}
-					else
-					{
-						SetVarNode = CompilerContext.SpawnIntermediateNode<UK2Node_CallFunction>(this, SourceGraph);
-					}
-					SetVarNode->SetFromFunction(SetByNameFunction);
-					SetVarNode->AllocateDefaultPins();
-
-					// Connect this node into the exec chain
-					UEdGraphPin* LastThen = LastNode->GetThenPin();
-					UEdGraphPin* SetVarExec = SetVarNode->GetExecPin();
-					LastThen->MakeLinkTo(SetVarExec);
-
-					// Connect the new actor to the 'object' pin
-					UEdGraphPin* ObjectPin = SetVarNode->FindPinChecked(ObjectParamName);
-					CallBeginResult->MakeLinkTo(ObjectPin);
-
-					// Fill in literal for 'property name' pin - name of pin is property name
-					UEdGraphPin* PropertyNamePin = SetVarNode->FindPinChecked(PropertyNameParamName);
-					PropertyNamePin->DefaultValue = SpawnVarPin->PinName;
-
-					// Move connection from the variable pin on the spawn node to the 'value' pin
-					UEdGraphPin* ValuePin = SetVarNode->FindPinChecked(ValueParamName);
-					CompilerContext.MovePinLinksToIntermediate(*SpawnVarPin, *ValuePin);
-					if(SpawnVarPin->PinType.bIsArray)
-					{
-						SetVarNode->PinConnectionListChanged(ValuePin);
-					}
-
-					// Update 'last node in sequence' var
-					LastNode = SetVarNode;
+					SetVarNode = CompilerContext.SpawnIntermediateNode<UK2Node_CallArrayFunction>(this, SourceGraph);
 				}
+				else
+				{
+					SetVarNode = CompilerContext.SpawnIntermediateNode<UK2Node_CallFunction>(this, SourceGraph);
+				}
+				SetVarNode->SetFromFunction(SetByNameFunction);
+				SetVarNode->AllocateDefaultPins();
+
+				// Connect this node into the exec chain
+				UEdGraphPin* LastThen = LastNode->GetThenPin();
+				UEdGraphPin* SetVarExec = SetVarNode->GetExecPin();
+				LastThen->MakeLinkTo(SetVarExec);
+
+				// Connect the new actor to the 'object' pin
+				UEdGraphPin* ObjectPin = SetVarNode->FindPinChecked(ObjectParamName);
+				CallBeginResult->MakeLinkTo(ObjectPin);
+
+				// Fill in literal for 'property name' pin - name of pin is property name
+				UEdGraphPin* PropertyNamePin = SetVarNode->FindPinChecked(PropertyNameParamName);
+				PropertyNamePin->DefaultValue = SpawnVarPin->PinName;
+
+				// Move connection from the variable pin on the spawn node to the 'value' pin
+				UEdGraphPin* ValuePin = SetVarNode->FindPinChecked(ValueParamName);
+				CompilerContext.MovePinLinksToIntermediate(*SpawnVarPin, *ValuePin);
+				if(SpawnVarPin->PinType.bIsArray)
+				{
+					SetVarNode->PinConnectionListChanged(ValuePin);
+				}
+
+				// Update 'last node in sequence' var
+				LastNode = SetVarNode;
 			}
 		}
+	}
 
-		// Make exec connection between 'then' on last node and 'finish'
-		UEdGraphPin* LastThen = LastNode->GetThenPin();
-		LastThen->MakeLinkTo(CallFinishExec);
+	// Make exec connection between 'then' on last node and 'finish'
+	UEdGraphPin* LastThen = LastNode->GetThenPin();
+	LastThen->MakeLinkTo(CallFinishExec);
 
-		// Break any links to the expanded node
-		BreakAllNodeLinks();
-	}
+	// Break any links to the expanded node
+	BreakAllNodeLinks();
 }
 
 bool UK2Node_SpawnActor::HasExternalBlueprintDependencies(TArray<class UStruct*>* OptionalOutput) const
diff --git a/Engine/Source/Editor/BlueprintGraph/Private/K2Node_SpawnActorFromClass.cpp b/Engine/Source/Editor/BlueprintGraph/Private/K2Node_SpawnActorFromClass.cpp
--- a/Engine/Source/Editor/BlueprintGraph/Private/K2Node_SpawnActorFromClass.cpp
+++ b/Engine/Source/Editor/BlueprintGraph/Private/K2Node_SpawnActorFromClass.cpp
@@ -371,115 +371,112 @@ void UK2Node_SpawnActorFromClass::ExpandNode(class FKismetCompilerContext& Compi
 {
 	Super::ExpandNode(CompilerContext, SourceGraph);
 
-	if (CompilerContext.bIsFullCompile)
+	static FName BeginSpawningBlueprintFuncName = GET_FUNCTION_NAME_CHECKED(UGameplayStatics, BeginSpawningActorFromClass);
+	static FString ActorClassParamName = FString(TEXT("ActorClass"));
+	static FString WorldContextParamName = FString(TEXT("WorldContextObject"));
+
+	static FName FinishSpawningFuncName = GET_FUNCTION_NAME_CHECKED(UGameplayStatics, FinishSpawningActor);
+	static FString ActorParamName = FString(TEXT("Actor"));
+	static FString TransformParamName = FString(TEXT("SpawnTransform"));
+	static FString NoCollisionFailParamName = FString(TEXT("bNoCollisionFail"));
+
+	static FString ObjectParamName = FString(TEXT("Object"));
+	static FString ValueParamName = FString(TEXT("Value"));
+	static FString PropertyNameParamName = FString(TEXT("PropertyName"));
+
+	UK2Node_SpawnActorFromClass* SpawnNode = this;
+	UEdGraphPin* SpawnNodeExec = SpawnNode->GetExecPin();
+	UEdGraphPin* SpawnNodeTransform = SpawnNode->GetSpawnTransformPin();
+	UEdGraphPin* SpawnNodeNoCollisionFail = GetNoCollisionFailPin();
+	UEdGraphPin* SpawnWorldContextPin = SpawnNode->GetWorldContextPin();
+	UEdGraphPin* SpawnClassPin = SpawnNode->GetClassPin();
+	UEdGraphPin* SpawnNodeThen = SpawnNode->GetThenPin();
+	UEdGraphPin* SpawnNodeResult = SpawnNode->GetResultPin();
+
+	UClass* SpawnClass = (SpawnClassPin != NULL) ? Cast<UClass>(SpawnClassPin->DefaultObject) : NULL;
+	if((0 == SpawnClassPin->LinkedTo.Num()) && (NULL == SpawnClass))
 	{
-		static FName BeginSpawningBlueprintFuncName = GET_FUNCTION_NAME_CHECKED(UGameplayStatics, BeginSpawningActorFromClass);
-		static FString ActorClassParamName = FString(TEXT("ActorClass"));
-		static FString WorldContextParamName = FString(TEXT("WorldContextObject"));
-
-		static FName FinishSpawningFuncName = GET_FUNCTION_NAME_CHECKED(UGameplayStatics, FinishSpawningActor);
-		static FString ActorParamName = FString(TEXT("Actor"));
-		static FString TransformParamName = FString(TEXT("SpawnTransform"));
-		static FString NoCollisionFailParamName = FString(TEXT("bNoCollisionFail"));
-
-		static FString ObjectParamName = FString(TEXT("Object"));
-		static FString ValueParamName = FString(TEXT("Value"));
-		static FString PropertyNameParamName = FString(TEXT("PropertyName"));
-
-		UK2Node_SpawnActorFromClass* SpawnNode = this;
-		UEdGraphPin* SpawnNodeExec = SpawnNode->GetExecPin();
-		UEdGraphPin* SpawnNodeTransform = SpawnNode->GetSpawnTransformPin();
-		UEdGraphPin* SpawnNodeNoCollisionFail = GetNoCollisionFailPin();
-		UEdGraphPin* SpawnWorldContextPin = SpawnNode->GetWorldContextPin();
-		UEdGraphPin* SpawnClassPin = SpawnNode->GetClassPin();
-		UEdGraphPin* SpawnNodeThen = SpawnNode->GetThenPin();
-		UEdGraphPin* SpawnNodeResult = SpawnNode->GetResultPin();
-
-		UClass* SpawnClass = (SpawnClassPin != NULL) ? Cast<UClass>(SpawnClassPin->DefaultObject) : NULL;
-		if((0 == SpawnClassPin->LinkedTo.Num()) && (NULL == SpawnClass))
-		{
-			CompilerContext.MessageLog.Error(*LOCTEXT("SpawnActorNodeMissingClass_Error", "Spawn node @@ must have a class specified.").ToString(), SpawnNode);
-			// we break exec links so this is the only error we get, don't want the SpawnActor node being considered and giving 'unexpected node' type warnings
-			SpawnNode->BreakAllNodeLinks();
-			return;
-		}
+		CompilerContext.MessageLog.Error(*LOCTEXT("SpawnActorNodeMissingClass_Error", "Spawn node @@ must have a class specified.").ToString(), SpawnNode);
+		// we break exec links so this is the only error we get, don't want the SpawnActor node being considered and giving 'unexpected node' type warnings
+		SpawnNode->BreakAllNodeLinks();
+		return;
+	}
 
-		//////////////////////////////////////////////////////////////////////////
-		// create 'begin spawn' call node
-		UK2Node_CallFunction* CallBeginSpawnNode = CompilerContext.SpawnIntermediateNode<UK2Node_CallFunction>(SpawnNode, SourceGraph);
-		CallBeginSpawnNode->FunctionReference.SetExternalMember(BeginSpawningBlueprintFuncName, UGameplayStatics::StaticClass());
-		CallBeginSpawnNode->AllocateDefaultPins();
+	//////////////////////////////////////////////////////////////////////////
+	// create 'begin spawn' call node
+	UK2Node_CallFunction* CallBeginSpawnNode = CompilerContext.SpawnIntermediateNode<UK2Node_CallFunction>(SpawnNode, SourceGraph);
+	CallBeginSpawnNode->FunctionReference.SetExternalMember(BeginSpawningBlueprintFuncName, UGameplayStatics::StaticClass());
+	CallBeginSpawnNode->AllocateDefaultPins();
 
-		UEdGraphPin* CallBeginExec = CallBeginSpawnNode->GetExecPin();
-		UEdGraphPin* CallBeginWorldContextPin = CallBeginSpawnNode->FindPinChecked(WorldContextParamName);
-		UEdGraphPin* CallBeginActorClassPin = CallBeginSpawnNode->FindPinChecked(ActorClassParamName);
-		UEdGraphPin* CallBeginTransform = CallBeginSpawnNode->FindPinChecked(TransformParamName);
-		UEdGraphPin* CallBeginNoCollisionFail = CallBeginSpawnNode->FindPinChecked(NoCollisionFailParamName);
-		UEdGraphPin* CallBeginResult = CallBeginSpawnNode->GetReturnValuePin();
+	UEdGraphPin* CallBeginExec = CallBeginSpawnNode->GetExecPin();
+	UEdGraphPin* CallBeginWorldContextPin = CallBeginSpawnNode->FindPinChecked(WorldContextParamName);
+	UEdGraphPin* CallBeginActorClassPin = CallBeginSpawnNode->FindPinChecked(ActorClassParamName);
+	UEdGraphPin* CallBeginTransform = CallBeginSpawnNode->FindPinChecked(TransformParamName);
+	UEdGraphPin* CallBeginNoCollisionFail = CallBeginSpawnNode->FindPinChecked(NoCollisionFailParamName);
+	UEdGraphPin* CallBeginResult = CallBeginSpawnNode->GetReturnValuePin();
 
-		// Move 'exec' connection from spawn node to 'begin spawn'
-		CompilerContext.MovePinLinksToIntermediate(*SpawnNodeExec, *CallBeginExec);
+	// Move 'exec' connection from spawn node to 'begin spawn'
+	CompilerContext.MovePinLinksToIntermediate(*SpawnNodeExec, *CallBeginExec);
 
-		if(SpawnClassPin->LinkedTo.Num() > 0)
-		{
-			// Copy the 'blueprint' connection from the spawn node to 'begin spawn'
-			CompilerContext.MovePinLinksToIntermediate(*SpawnClassPin, *CallBeginActorClassPin);
-		}
-		else
-		{
-			// Copy blueprint literal onto begin spawn call 
-			CallBeginActorClassPin->DefaultObject = SpawnClass;
-		}
+	if(SpawnClassPin->LinkedTo.Num() > 0)
+	{
+		// Copy the 'blueprint' connection from the spawn node to 'begin spawn'
+		CompilerContext.MovePinLinksToIntermediate(*SpawnClassPin, *CallBeginActorClassPin);
+	}
+	else
+	{
+		// Copy blueprint literal onto begin spawn call 
+		CallBeginActorClassPin->DefaultObject = SpawnClass;
+	}
 
-		// Copy the world context connection from the spawn node to 'begin spawn' if necessary
-		if (SpawnWorldContextPin)
-		{
-			CompilerContext.MovePinLinksToIntermediate(*SpawnWorldContextPin, *CallBeginWorldContextPin);
-		}
+	// Copy the world context connection from the spawn node to 'begin spawn' if necessary
+	if (SpawnWorldContextPin)
+	{
+		CompilerContext.MovePinLinksToIntermediate(*SpawnWorldContextPin, *CallBeginWorldContextPin);
+	}
 
-		// Copy the 'transform' connection from the spawn node to 'begin spawn'
-		CompilerContext.MovePinLinksToIntermediate(*SpawnNodeTransform, *CallBeginTransform);
+	// Copy the 'transform' connection from the spawn node to 'begin spawn'
+	CompilerContext.MovePinLinksToIntermediate(*SpawnNodeTransform, *CallBeginTransform);
 
-		// Copy the 'bNoCollisionFail' connection from the spawn node to 'begin spawn'
-		CompilerContext.MovePinLinksToIntermediate(*SpawnNodeNoCollisionFail, *CallBeginNoCollisionFail);
+	// Copy the 'bNoCollisionFail' connection from the spawn node to 'begin spawn'
+	CompilerContext.MovePinLinksToIntermediate(*SpawnNodeNoCollisionFail, *CallBeginNoCollisionFail);
 
-		//////////////////////////////////////////////////////////////////////////
-		// create 'finish spawn' call node
-		UK2Node_CallFunction* CallFinishSpawnNode = CompilerContext.SpawnIntermediateNode<UK2Node_CallFunction>(SpawnNode, SourceGraph);
-		CallFinishSpawnNode->FunctionReference.SetExternalMember(FinishSpawningFuncName, UGameplayStatics::StaticClass());
-		CallFinishSpawnNode->AllocateDefaultPins();
+	//////////////////////////////////////////////////////////////////////////
+	// create 'finish spawn' call node
+	UK2Node_CallFunction* CallFinishSpawnNode = CompilerContext.SpawnIntermediateNode<UK2Node_CallFunction>(SpawnNode, SourceGraph);
+	CallFinishSpawnNode->FunctionReference.SetExternalMember(FinishSpawningFuncName, UGameplayStatics::StaticClass());
+	CallFinishSpawnNode->AllocateDefaultPins();
 
-		UEdGraphPin* CallFinishExec = CallFinishSpawnNode->GetExecPin();
-		UEdGraphPin* CallFinishThen = CallFinishSpawnNode->GetThenPin();
-		UEdGraphPin* CallFinishActor = CallFinishSpawnNode->FindPinChecked(ActorParamName);
-		UEdGraphPin* CallFinishTransform = CallFinishSpawnNode->FindPinChecked(TransformParamName);
-		UEdGraphPin* CallFinishResult = CallFinishSpawnNode->GetReturnValuePin();
+	UEdGraphPin* CallFinishExec = CallFinishSpawnNode->GetExecPin();
+	UEdGraphPin* CallFinishThen = CallFinishSpawnNode->GetThenPin();
+	UEdGraphPin* CallFinishActor = CallFinishSpawnNode->FindPinChecked(ActorParamName);
+	UEdGraphPin* CallFinishTransform = CallFinishSpawnNode->FindPinChecked(TransformParamName);
+	UEdGraphPin* CallFinishResult = CallFinishSpawnNode->GetReturnValuePin();
 
-		// Move 'then' connection from spawn node to 'finish spawn'
-		CompilerContext.MovePinLinksToIntermediate(*SpawnNodeThen, *CallFinishThen);
+	// Move 'then' connection from spawn node to 'finish spawn'
+	CompilerContext.MovePinLinksToIntermediate(*SpawnNodeThen, *CallFinishThen);
 
-		// Copy transform connection
-		CompilerContext.CopyPinLinksToIntermediate(*CallBeginTransform, *CallFinishTransform);
+	// Copy transform connection
+	CompilerContext.CopyPinLinksToIntermediate(*CallBeginTransform, *CallFinishTransform);
 
-		// Connect output actor from 'begin' to 'finish'
-		CallBeginResult->MakeLinkTo(CallFinishActor);
+	// Connect output actor from 'begin' to 'finish'
+	CallBeginResult->MakeLinkTo(CallFinishActor);
 
-		// Move result connection from spawn node to 'finish spawn'
-		CallFinishResult->PinType = SpawnNodeResult->PinType; // Copy type so it uses the right actor subclass
-		CompilerContext.MovePinLinksToIntermediate(*SpawnNodeResult, *CallFinishResult);
+	// Move result connection from spawn node to 'finish spawn'
+	CallFinishResult->PinType = SpawnNodeResult->PinType; // Copy type so it uses the right actor subclass
+	CompilerContext.MovePinLinksToIntermediate(*SpawnNodeResult, *CallFinishResult);
 
-		//////////////////////////////////////////////////////////////////////////
-		// create 'set var' nodes
+	//////////////////////////////////////////////////////////////////////////
+	// create 'set var' nodes
 
-		// Get 'result' pin from 'begin spawn', this is the actual actor we want to set properties on
-		UEdGraphPin* LastThen = FKismetCompilerUtilities::GenerateAssignmentNodes(CompilerContext, SourceGraph, CallBeginSpawnNode, SpawnNode, CallBeginResult, GetClassToSpawn() );
+	// Get 'result' pin from 'begin spawn', this is the actual actor we want to set properties on
+	UEdGraphPin* LastThen = FKismetCompilerUtilities::GenerateAssignmentNodes(CompilerContext, SourceGraph, CallBeginSpawnNode, SpawnNode, CallBeginResult, GetClassToSpawn() );
 
-		// Make exec connection between 'then' on last node and 'finish'
-		LastThen->MakeLinkTo(CallFinishExec);
+	// Make exec connection between 'then' on last node and 'finish'
+	LastThen->MakeLinkTo(CallFinishExec);
 
-		// Break any links to the expanded node
-		SpawnNode->BreakAllNodeLinks();
-	}
+	// Break any links to the expanded node
+	SpawnNode->BreakAllNodeLinks();
 }
 
 bool UK2Node_SpawnActorFromClass::HasExternalBlueprintDependencies(TArray<class UStruct*>* OptionalOutput) const
diff --git a/Engine/Source/Editor/GameplayAbilitiesEditor/Private/K2Node_LatentAbilityCall.cpp b/Engine/Source/Editor/GameplayAbilitiesEditor/Private/K2Node_LatentAbilityCall.cpp
--- a/Engine/Source/Editor/GameplayAbilitiesEditor/Private/K2Node_LatentAbilityCall.cpp
+++ b/Engine/Source/Editor/GameplayAbilitiesEditor/Private/K2Node_LatentAbilityCall.cpp
@@ -570,11 +570,6 @@ void UK2Node_LatentAbilityCall::ExpandNode(class FKismetCompilerContext& Compile
 
 	UK2Node::ExpandNode(CompilerContext, SourceGraph);
 
-	if (!CompilerContext.bIsFullCompile)
-	{
-		return;
-	}
-
 	if (!validatedActorSpawn && !validatedActorArraySpawn)
 	{
 		ValidateActorSpawning(CompilerContext, true);
diff --git a/Engine/Source/Editor/GameplayTagsEditor/Private/GameplayTagsK2Node_LiteralGameplayTag.cpp b/Engine/Source/Editor/GameplayTagsEditor/Private/GameplayTagsK2Node_LiteralGameplayTag.cpp
--- a/Engine/Source/Editor/GameplayTagsEditor/Private/GameplayTagsK2Node_LiteralGameplayTag.cpp
+++ b/Engine/Source/Editor/GameplayTagsEditor/Private/GameplayTagsK2Node_LiteralGameplayTag.cpp
@@ -58,93 +58,90 @@ void UGameplayTagsK2Node_LiteralGameplayTag::ExpandNode(class FKismetCompilerCon
 {
 	Super::ExpandNode(CompilerContext, SourceGraph);
 
-	if (CompilerContext.bIsFullCompile)
-	{
-		const UEdGraphSchema_K2* Schema = CompilerContext.GetSchema();
+	const UEdGraphSchema_K2* Schema = CompilerContext.GetSchema();
 
-		// Get The input and output pins to our node
-		UEdGraphPin* TagInPin = FindPin(TEXT("TagIn"));
-		UEdGraphPin* TagOutPin = FindPinChecked(Schema->PN_ReturnValue);
+	// Get The input and output pins to our node
+	UEdGraphPin* TagInPin = FindPin(TEXT("TagIn"));
+	UEdGraphPin* TagOutPin = FindPinChecked(Schema->PN_ReturnValue);
 
-		// Create a Make Struct
-		UK2Node_MakeStruct* MakeStructNode = SourceGraph->CreateBlankNode<UK2Node_MakeStruct>();
-		MakeStructNode->StructType = FGameplayTagContainer::StaticStruct();
-		MakeStructNode->AllocateDefaultPins();
+	// Create a Make Struct
+	UK2Node_MakeStruct* MakeStructNode = SourceGraph->CreateBlankNode<UK2Node_MakeStruct>();
+	MakeStructNode->StructType = FGameplayTagContainer::StaticStruct();
+	MakeStructNode->AllocateDefaultPins();
 
-		// Create a Make Array
-		UK2Node_MakeArray* MakeArrayNode = SourceGraph->CreateBlankNode<UK2Node_MakeArray>();
-		MakeArrayNode->AllocateDefaultPins();
+	// Create a Make Array
+	UK2Node_MakeArray* MakeArrayNode = SourceGraph->CreateBlankNode<UK2Node_MakeArray>();
+	MakeArrayNode->AllocateDefaultPins();
 		
-		// Connect the output of our MakeArray to the Input of our MakeStruct so it sets the Array Type
-		UEdGraphPin* InPin = MakeStructNode->FindPin( TEXT("GameplayTags") );
-		if( InPin )
+	// Connect the output of our MakeArray to the Input of our MakeStruct so it sets the Array Type
+	UEdGraphPin* InPin = MakeStructNode->FindPin( TEXT("GameplayTags") );
+	if( InPin )
+	{
+		InPin->MakeLinkTo( MakeArrayNode->GetOutputPin() );
+	}
+
+	// Add the FName Values to the MakeArray input pins
+	UEdGraphPin* ArrayInputPin = NULL;
+	FString TagString = TagInPin->GetDefaultAsString();
+
+	if( TagString.StartsWith( TEXT("(") ) && TagString.EndsWith( TEXT(")") ) )
+	{
+		TagString = TagString.LeftChop(1);
+		TagString = TagString.RightChop(1);
+		TagString.Split("=", NULL, &TagString);
+		TagString = TagString.LeftChop(1);
+		TagString = TagString.RightChop(1);
+
+		FString ReadTag;
+		FString Remainder;
+		int32 MakeIndex = 0;
+		while( TagString.Split( TEXT(","), &ReadTag, &Remainder ) ) 
 		{
-			InPin->MakeLinkTo( MakeArrayNode->GetOutputPin() );
-		}
+			TagString = Remainder;
 
-		// Add the FName Values to the MakeArray input pins
-		UEdGraphPin* ArrayInputPin = NULL;
-		FString TagString = TagInPin->GetDefaultAsString();
+			ArrayInputPin = MakeArrayNode->FindPin( FString::Printf( TEXT("[%d]"), MakeIndex ) );
+			ArrayInputPin->PinType.PinCategory = TEXT("struct");
+			ArrayInputPin->PinType.PinSubCategoryObject = FGameplayTag::StaticStruct();
+			ArrayInputPin->DefaultValue = ReadTag;
 
-		if( TagString.StartsWith( TEXT("(") ) && TagString.EndsWith( TEXT(")") ) )
+			MakeIndex++;
+			MakeArrayNode->AddInputPin();
+		}
+		if( Remainder.IsEmpty() )
 		{
-			TagString = TagString.LeftChop(1);
-			TagString = TagString.RightChop(1);
-			TagString.Split("=", NULL, &TagString);
-			TagString = TagString.LeftChop(1);
-			TagString = TagString.RightChop(1);
-
-			FString ReadTag;
-			FString Remainder;
-			int32 MakeIndex = 0;
-			while( TagString.Split( TEXT(","), &ReadTag, &Remainder ) ) 
-			{
-				TagString = Remainder;
-
-				ArrayInputPin = MakeArrayNode->FindPin( FString::Printf( TEXT("[%d]"), MakeIndex ) );
-				ArrayInputPin->PinType.PinCategory = TEXT("struct");
-				ArrayInputPin->PinType.PinSubCategoryObject = FGameplayTag::StaticStruct();
-				ArrayInputPin->DefaultValue = ReadTag;
-
-				MakeIndex++;
-				MakeArrayNode->AddInputPin();
-			}
-			if( Remainder.IsEmpty() )
-			{
-				Remainder = TagString;
-			}
-			if( !Remainder.IsEmpty() )
-			{
-				ArrayInputPin = MakeArrayNode->FindPin( FString::Printf( TEXT("[%d]"), MakeIndex ) );
-				ArrayInputPin->PinType.PinCategory = TEXT("struct");
-				ArrayInputPin->PinType.PinSubCategoryObject = FGameplayTag::StaticStruct();
-				ArrayInputPin->DefaultValue = Remainder;
-
-				MakeArrayNode->PostReconstructNode();
-			}
-			else
-			{
-				MakeArrayNode->RemoveInputPin(MakeArrayNode->FindPin(TEXT("[0]")));
-				MakeArrayNode->PostReconstructNode();
-			}
+			Remainder = TagString;
 		}
-		else
+		if( !Remainder.IsEmpty() )
 		{
-			MakeArrayNode->RemoveInputPin(MakeArrayNode->FindPin(TEXT("[0]")));
+			ArrayInputPin = MakeArrayNode->FindPin( FString::Printf( TEXT("[%d]"), MakeIndex ) );
+			ArrayInputPin->PinType.PinCategory = TEXT("struct");
+			ArrayInputPin->PinType.PinSubCategoryObject = FGameplayTag::StaticStruct();
+			ArrayInputPin->DefaultValue = Remainder;
+
 			MakeArrayNode->PostReconstructNode();
 		}
-
-		// Move the Output of the MakeArray to the Output of our node
-		UEdGraphPin* OutPin = MakeStructNode->FindPin( MakeStructNode->StructType->GetName() );
-		if( OutPin && TagOutPin )
+		else
 		{
-			OutPin->PinType = TagOutPin->PinType; // Copy type so it uses the right actor subclass
-			CompilerContext.MovePinLinksToIntermediate(*TagOutPin, *OutPin);
+			MakeArrayNode->RemoveInputPin(MakeArrayNode->FindPin(TEXT("[0]")));
+			MakeArrayNode->PostReconstructNode();
 		}
+	}
+	else
+	{
+		MakeArrayNode->RemoveInputPin(MakeArrayNode->FindPin(TEXT("[0]")));
+		MakeArrayNode->PostReconstructNode();
+	}
 
-		// Break any links to the expanded node
-		BreakAllNodeLinks();
+	// Move the Output of the MakeArray to the Output of our node
+	UEdGraphPin* OutPin = MakeStructNode->FindPin( MakeStructNode->StructType->GetName() );
+	if( OutPin && TagOutPin )
+	{
+		OutPin->PinType = TagOutPin->PinType; // Copy type so it uses the right actor subclass
+		CompilerContext.MovePinLinksToIntermediate(*TagOutPin, *OutPin);
 	}
+
+	// Break any links to the expanded node
+	BreakAllNodeLinks();
 }
 
 void UGameplayTagsK2Node_LiteralGameplayTag::GetMenuActions(FBlueprintActionDatabaseRegistrar& ActionRegistrar) const
diff --git a/Engine/Source/Editor/UMGEditor/Private/Nodes/K2Node_CreateDragDropOperation.cpp b/Engine/Source/Editor/UMGEditor/Private/Nodes/K2Node_CreateDragDropOperation.cpp
--- a/Engine/Source/Editor/UMGEditor/Private/Nodes/K2Node_CreateDragDropOperation.cpp
+++ b/Engine/Source/Editor/UMGEditor/Private/Nodes/K2Node_CreateDragDropOperation.cpp
@@ -70,66 +70,63 @@ void UK2Node_CreateDragDropOperation::ExpandNode(class FKismetCompilerContext& C
 {
 	Super::ExpandNode(CompilerContext, SourceGraph);
 
-	if ( CompilerContext.bIsFullCompile )
+	static FName Create_FunctionName = GET_FUNCTION_NAME_CHECKED(UWidgetBlueprintLibrary, CreateDragDropOperation);
+	static FString OperationClass_ParamName = FString(TEXT("OperationClass"));
+
+	UK2Node_CreateDragDropOperation* CreateOpNode = this;
+	UEdGraphPin* SpawnNodeExec = CreateOpNode->GetExecPin();
+	UEdGraphPin* SpawnClassPin = CreateOpNode->GetClassPin();
+	UEdGraphPin* SpawnNodeThen = CreateOpNode->GetThenPin();
+	UEdGraphPin* SpawnNodeResult = CreateOpNode->GetResultPin();
+
+	UClass* SpawnClass = ( SpawnClassPin != NULL ) ? Cast<UClass>(SpawnClassPin->DefaultObject) : NULL;
+	//if ( ( 0 == SpawnClassPin->LinkedTo.Num() ) && ( NULL == SpawnClass ) )
+	//{
+	//	CompilerContext.MessageLog.Error(*LOCTEXT("CreateWidgetNodeMissingClass_Error", "Create Drag/Drop node @@ must have a class specified.").ToString(), CreateOpNode);
+	//	// we break exec links so this is the only error we get, don't want the CreateWidget node being considered and giving 'unexpected node' type warnings
+	//	CreateOpNode->BreakAllNodeLinks();
+	//	return;
+	//}
+
+	//////////////////////////////////////////////////////////////////////////
+	// create 'UWidgetBlueprintLibrary::CreateDragDropOperation' call node
+	UK2Node_CallFunction* CallCreateNode = CompilerContext.SpawnIntermediateNode<UK2Node_CallFunction>(CreateOpNode, SourceGraph);
+	CallCreateNode->FunctionReference.SetExternalMember(Create_FunctionName, UWidgetBlueprintLibrary::StaticClass());
+	CallCreateNode->AllocateDefaultPins();
+
+	UEdGraphPin* CallCreateExec = CallCreateNode->GetExecPin();
+	UEdGraphPin* CallCreateOperationClassPin = CallCreateNode->FindPinChecked(OperationClass_ParamName);
+	UEdGraphPin* CallCreateResult = CallCreateNode->GetReturnValuePin();
+
+	// Move 'exec' connection from create widget node to 'UWidgetBlueprintLibrary::CreateDragDropOperation'
+	CompilerContext.MovePinLinksToIntermediate(*SpawnNodeExec, *CallCreateExec);
+
+	if ( SpawnClassPin->LinkedTo.Num() > 0 )
 	{
-		static FName Create_FunctionName = GET_FUNCTION_NAME_CHECKED(UWidgetBlueprintLibrary, CreateDragDropOperation);
-		static FString OperationClass_ParamName = FString(TEXT("OperationClass"));
-
-		UK2Node_CreateDragDropOperation* CreateOpNode = this;
-		UEdGraphPin* SpawnNodeExec = CreateOpNode->GetExecPin();
-		UEdGraphPin* SpawnClassPin = CreateOpNode->GetClassPin();
-		UEdGraphPin* SpawnNodeThen = CreateOpNode->GetThenPin();
-		UEdGraphPin* SpawnNodeResult = CreateOpNode->GetResultPin();
-
-		UClass* SpawnClass = ( SpawnClassPin != NULL ) ? Cast<UClass>(SpawnClassPin->DefaultObject) : NULL;
-		//if ( ( 0 == SpawnClassPin->LinkedTo.Num() ) && ( NULL == SpawnClass ) )
-		//{
-		//	CompilerContext.MessageLog.Error(*LOCTEXT("CreateWidgetNodeMissingClass_Error", "Create Drag/Drop node @@ must have a class specified.").ToString(), CreateOpNode);
-		//	// we break exec links so this is the only error we get, don't want the CreateWidget node being considered and giving 'unexpected node' type warnings
-		//	CreateOpNode->BreakAllNodeLinks();
-		//	return;
-		//}
-
-		//////////////////////////////////////////////////////////////////////////
-		// create 'UWidgetBlueprintLibrary::CreateDragDropOperation' call node
-		UK2Node_CallFunction* CallCreateNode = CompilerContext.SpawnIntermediateNode<UK2Node_CallFunction>(CreateOpNode, SourceGraph);
-		CallCreateNode->FunctionReference.SetExternalMember(Create_FunctionName, UWidgetBlueprintLibrary::StaticClass());
-		CallCreateNode->AllocateDefaultPins();
-
-		UEdGraphPin* CallCreateExec = CallCreateNode->GetExecPin();
-		UEdGraphPin* CallCreateOperationClassPin = CallCreateNode->FindPinChecked(OperationClass_ParamName);
-		UEdGraphPin* CallCreateResult = CallCreateNode->GetReturnValuePin();
-
-		// Move 'exec' connection from create widget node to 'UWidgetBlueprintLibrary::CreateDragDropOperation'
-		CompilerContext.MovePinLinksToIntermediate(*SpawnNodeExec, *CallCreateExec);
-
-		if ( SpawnClassPin->LinkedTo.Num() > 0 )
-		{
-			// Copy the 'blueprint' connection from the spawn node to 'UWidgetBlueprintLibrary::CreateDragDropOperation'
-			CompilerContext.MovePinLinksToIntermediate(*SpawnClassPin, *CallCreateOperationClassPin);
-		}
-		else
-		{
-			// Copy blueprint literal onto 'UWidgetBlueprintLibrary::CreateDragDropOperation' call 
-			CallCreateOperationClassPin->DefaultObject = SpawnClass;
-		}
+		// Copy the 'blueprint' connection from the spawn node to 'UWidgetBlueprintLibrary::CreateDragDropOperation'
+		CompilerContext.MovePinLinksToIntermediate(*SpawnClassPin, *CallCreateOperationClassPin);
+	}
+	else
+	{
+		// Copy blueprint literal onto 'UWidgetBlueprintLibrary::CreateDragDropOperation' call 
+		CallCreateOperationClassPin->DefaultObject = SpawnClass;
+	}
 
-		// Move result connection from spawn node to 'UWidgetBlueprintLibrary::CreateDragDropOperation'
-		CallCreateResult->PinType = SpawnNodeResult->PinType; // Copy type so it uses the right actor subclass
-		CompilerContext.MovePinLinksToIntermediate(*SpawnNodeResult, *CallCreateResult);
+	// Move result connection from spawn node to 'UWidgetBlueprintLibrary::CreateDragDropOperation'
+	CallCreateResult->PinType = SpawnNodeResult->PinType; // Copy type so it uses the right actor subclass
+	CompilerContext.MovePinLinksToIntermediate(*SpawnNodeResult, *CallCreateResult);
 
-		//////////////////////////////////////////////////////////////////////////
-		// create 'set var' nodes
+	//////////////////////////////////////////////////////////////////////////
+	// create 'set var' nodes
 
-		// Get 'result' pin from 'begin spawn', this is the actual actor we want to set properties on
-		UEdGraphPin* LastThen = FKismetCompilerUtilities::GenerateAssignmentNodes(CompilerContext, SourceGraph, CallCreateNode, CreateOpNode, CallCreateResult, GetClassToSpawn());
+	// Get 'result' pin from 'begin spawn', this is the actual actor we want to set properties on
+	UEdGraphPin* LastThen = FKismetCompilerUtilities::GenerateAssignmentNodes(CompilerContext, SourceGraph, CallCreateNode, CreateOpNode, CallCreateResult, GetClassToSpawn());
 
-		// Move 'then' connection from create widget node to the last 'then'
-		CompilerContext.MovePinLinksToIntermediate(*SpawnNodeThen, *LastThen);
+	// Move 'then' connection from create widget node to the last 'then'
+	CompilerContext.MovePinLinksToIntermediate(*SpawnNodeThen, *LastThen);
 
-		// Break any links to the expanded node
-		CreateOpNode->BreakAllNodeLinks();
-	}
+	// Break any links to the expanded node
+	CreateOpNode->BreakAllNodeLinks();
 }
 
 #undef LOCTEXT_NAMESPACE
diff --git a/Engine/Source/Editor/UMGEditor/Private/Nodes/K2Node_CreateWidget.cpp b/Engine/Source/Editor/UMGEditor/Private/Nodes/K2Node_CreateWidget.cpp
--- a/Engine/Source/Editor/UMGEditor/Private/Nodes/K2Node_CreateWidget.cpp
+++ b/Engine/Source/Editor/UMGEditor/Private/Nodes/K2Node_CreateWidget.cpp
@@ -90,81 +90,78 @@ void UK2Node_CreateWidget::ExpandNode(class FKismetCompilerContext& CompilerCont
 {
 	Super::ExpandNode(CompilerContext, SourceGraph);
 
-	if ( CompilerContext.bIsFullCompile )
+	static FName Create_FunctionName = GET_FUNCTION_NAME_CHECKED(UWidgetBlueprintLibrary, Create);
+	static FString WorldContextObject_ParamName = FString(TEXT("WorldContextObject"));
+	static FString WidgetType_ParamName = FString(TEXT("WidgetType"));
+	static FString OwningPlayer_ParamName = FString(TEXT("OwningPlayer"));
+
+	UK2Node_CreateWidget* CreateWidgetNode = this;
+	UEdGraphPin* SpawnNodeExec = CreateWidgetNode->GetExecPin();
+	UEdGraphPin* SpawnWorldContextPin = CreateWidgetNode->GetWorldContextPin();
+	UEdGraphPin* SpawnOwningPlayerPin = CreateWidgetNode->GetOwningPlayerPin();
+	UEdGraphPin* SpawnClassPin = CreateWidgetNode->GetClassPin();
+	UEdGraphPin* SpawnNodeThen = CreateWidgetNode->GetThenPin();
+	UEdGraphPin* SpawnNodeResult = CreateWidgetNode->GetResultPin();
+
+	UClass* SpawnClass = ( SpawnClassPin != NULL ) ? Cast<UClass>(SpawnClassPin->DefaultObject) : NULL;
+	if ( ( 0 == SpawnClassPin->LinkedTo.Num() ) && ( NULL == SpawnClass ) )
 	{
-		static FName Create_FunctionName = GET_FUNCTION_NAME_CHECKED(UWidgetBlueprintLibrary, Create);
-		static FString WorldContextObject_ParamName = FString(TEXT("WorldContextObject"));
-		static FString WidgetType_ParamName = FString(TEXT("WidgetType"));
-		static FString OwningPlayer_ParamName = FString(TEXT("OwningPlayer"));
-
-		UK2Node_CreateWidget* CreateWidgetNode = this;
-		UEdGraphPin* SpawnNodeExec = CreateWidgetNode->GetExecPin();
-		UEdGraphPin* SpawnWorldContextPin = CreateWidgetNode->GetWorldContextPin();
-		UEdGraphPin* SpawnOwningPlayerPin = CreateWidgetNode->GetOwningPlayerPin();
-		UEdGraphPin* SpawnClassPin = CreateWidgetNode->GetClassPin();
-		UEdGraphPin* SpawnNodeThen = CreateWidgetNode->GetThenPin();
-		UEdGraphPin* SpawnNodeResult = CreateWidgetNode->GetResultPin();
-
-		UClass* SpawnClass = ( SpawnClassPin != NULL ) ? Cast<UClass>(SpawnClassPin->DefaultObject) : NULL;
-		if ( ( 0 == SpawnClassPin->LinkedTo.Num() ) && ( NULL == SpawnClass ) )
-		{
-			CompilerContext.MessageLog.Error(*LOCTEXT("CreateWidgetNodeMissingClass_Error", "Spawn node @@ must have a class specified.").ToString(), CreateWidgetNode);
-			// we break exec links so this is the only error we get, don't want the CreateWidget node being considered and giving 'unexpected node' type warnings
-			CreateWidgetNode->BreakAllNodeLinks();
-			return;
-		}
-
-		//////////////////////////////////////////////////////////////////////////
-		// create 'UWidgetBlueprintLibrary::Create' call node
-		UK2Node_CallFunction* CallCreateNode = CompilerContext.SpawnIntermediateNode<UK2Node_CallFunction>(CreateWidgetNode, SourceGraph);
-		CallCreateNode->FunctionReference.SetExternalMember(Create_FunctionName, UWidgetBlueprintLibrary::StaticClass());
-		CallCreateNode->AllocateDefaultPins();
-
-		UEdGraphPin* CallCreateExec = CallCreateNode->GetExecPin();
-		UEdGraphPin* CallCreateWorldContextPin = CallCreateNode->FindPinChecked(WorldContextObject_ParamName);
-		UEdGraphPin* CallCreateWidgetTypePin = CallCreateNode->FindPinChecked(WidgetType_ParamName);
-		UEdGraphPin* CallCreateOwningPlayerPin = CallCreateNode->FindPinChecked(OwningPlayer_ParamName);
-		UEdGraphPin* CallCreateResult = CallCreateNode->GetReturnValuePin();
-
-		// Move 'exec' connection from create widget node to 'UWidgetBlueprintLibrary::Create'
-		CompilerContext.MovePinLinksToIntermediate(*SpawnNodeExec, *CallCreateExec);
-
-		if ( SpawnClassPin->LinkedTo.Num() > 0 )
-		{
-			// Copy the 'blueprint' connection from the spawn node to 'UWidgetBlueprintLibrary::Create'
-			CompilerContext.MovePinLinksToIntermediate(*SpawnClassPin, *CallCreateWidgetTypePin);
-		}
-		else
-		{
-			// Copy blueprint literal onto 'UWidgetBlueprintLibrary::Create' call 
-			CallCreateWidgetTypePin->DefaultObject = SpawnClass;
-		}
-
-		// Copy the world context connection from the spawn node to 'UWidgetBlueprintLibrary::Create' if necessary
-		if ( SpawnWorldContextPin )
-		{
-			CompilerContext.MovePinLinksToIntermediate(*SpawnWorldContextPin, *CallCreateWorldContextPin);
-		}
-
-		// Copy the 'Owning Player' connection from the spawn node to 'UWidgetBlueprintLibrary::Create'
-		CompilerContext.MovePinLinksToIntermediate(*SpawnOwningPlayerPin, *CallCreateOwningPlayerPin);
-
-		// Move result connection from spawn node to 'UWidgetBlueprintLibrary::Create'
-		CallCreateResult->PinType = SpawnNodeResult->PinType; // Copy type so it uses the right actor subclass
-		CompilerContext.MovePinLinksToIntermediate(*SpawnNodeResult, *CallCreateResult);
-
-		//////////////////////////////////////////////////////////////////////////
-		// create 'set var' nodes
-
-		// Get 'result' pin from 'begin spawn', this is the actual actor we want to set properties on
-		UEdGraphPin* LastThen = FKismetCompilerUtilities::GenerateAssignmentNodes(CompilerContext, SourceGraph, CallCreateNode, CreateWidgetNode, CallCreateResult, GetClassToSpawn());
-
-		// Move 'then' connection from create widget node to the last 'then'
-		CompilerContext.MovePinLinksToIntermediate(*SpawnNodeThen, *LastThen);
-
-		// Break any links to the expanded node
+		CompilerContext.MessageLog.Error(*LOCTEXT("CreateWidgetNodeMissingClass_Error", "Spawn node @@ must have a class specified.").ToString(), CreateWidgetNode);
+		// we break exec links so this is the only error we get, don't want the CreateWidget node being considered and giving 'unexpected node' type warnings
 		CreateWidgetNode->BreakAllNodeLinks();
+		return;
 	}
+
+	//////////////////////////////////////////////////////////////////////////
+	// create 'UWidgetBlueprintLibrary::Create' call node
+	UK2Node_CallFunction* CallCreateNode = CompilerContext.SpawnIntermediateNode<UK2Node_CallFunction>(CreateWidgetNode, SourceGraph);
+	CallCreateNode->FunctionReference.SetExternalMember(Create_FunctionName, UWidgetBlueprintLibrary::StaticClass());
+	CallCreateNode->AllocateDefaultPins();
+
+	UEdGraphPin* CallCreateExec = CallCreateNode->GetExecPin();
+	UEdGraphPin* CallCreateWorldContextPin = CallCreateNode->FindPinChecked(WorldContextObject_ParamName);
+	UEdGraphPin* CallCreateWidgetTypePin = CallCreateNode->FindPinChecked(WidgetType_ParamName);
+	UEdGraphPin* CallCreateOwningPlayerPin = CallCreateNode->FindPinChecked(OwningPlayer_ParamName);
+	UEdGraphPin* CallCreateResult = CallCreateNode->GetReturnValuePin();
+
+	// Move 'exec' connection from create widget node to 'UWidgetBlueprintLibrary::Create'
+	CompilerContext.MovePinLinksToIntermediate(*SpawnNodeExec, *CallCreateExec);
+
+	if ( SpawnClassPin->LinkedTo.Num() > 0 )
+	{
+		// Copy the 'blueprint' connection from the spawn node to 'UWidgetBlueprintLibrary::Create'
+		CompilerContext.MovePinLinksToIntermediate(*SpawnClassPin, *CallCreateWidgetTypePin);
+	}
+	else
+	{
+		// Copy blueprint literal onto 'UWidgetBlueprintLibrary::Create' call 
+		CallCreateWidgetTypePin->DefaultObject = SpawnClass;
+	}
+
+	// Copy the world context connection from the spawn node to 'UWidgetBlueprintLibrary::Create' if necessary
+	if ( SpawnWorldContextPin )
+	{
+		CompilerContext.MovePinLinksToIntermediate(*SpawnWorldContextPin, *CallCreateWorldContextPin);
+	}
+
+	// Copy the 'Owning Player' connection from the spawn node to 'UWidgetBlueprintLibrary::Create'
+	CompilerContext.MovePinLinksToIntermediate(*SpawnOwningPlayerPin, *CallCreateOwningPlayerPin);
+
+	// Move result connection from spawn node to 'UWidgetBlueprintLibrary::Create'
+	CallCreateResult->PinType = SpawnNodeResult->PinType; // Copy type so it uses the right actor subclass
+	CompilerContext.MovePinLinksToIntermediate(*SpawnNodeResult, *CallCreateResult);
+
+	//////////////////////////////////////////////////////////////////////////
+	// create 'set var' nodes
+
+	// Get 'result' pin from 'begin spawn', this is the actual actor we want to set properties on
+	UEdGraphPin* LastThen = FKismetCompilerUtilities::GenerateAssignmentNodes(CompilerContext, SourceGraph, CallCreateNode, CreateWidgetNode, CallCreateResult, GetClassToSpawn());
+
+	// Move 'then' connection from create widget node to the last 'then'
+	CompilerContext.MovePinLinksToIntermediate(*SpawnNodeThen, *LastThen);
+
+	// Break any links to the expanded node
+	CreateWidgetNode->BreakAllNodeLinks();
 }
 
 #undef LOCTEXT_NAMESPACE
