Converted the LiveEditor to use FText with Slate

LiveEditor now makes use of LOCTEXT to support localized text, rather than using FString. This is required to remove SLATE_TEXT_ATTRIBUTE from the Slate widgets.

diff --git a/Engine/Plugins/Experimental/LiveEditor/Source/LiveEditor/Private/LiveEditorBlueprintBindingWizard.cpp b/Engine/Plugins/Experimental/LiveEditor/Source/LiveEditor/Private/LiveEditorBlueprintBindingWizard.cpp
--- a/Engine/Plugins/Experimental/LiveEditor/Source/LiveEditor/Private/LiveEditorBlueprintBindingWizard.cpp
+++ b/Engine/Plugins/Experimental/LiveEditor/Source/LiveEditor/Private/LiveEditorBlueprintBindingWizard.cpp
@@ -2,6 +2,8 @@
 #include "LiveEditorUserData.h"
 #include "LiveEditorBlueprintBindingWizard.h"
 
+#define LOCTEXT_NAMESPACE "LiveEditorWizard"
+
 class SLiveEditorBlueprintBindingWindow : public SCompoundWidget
 {
 public:
@@ -10,11 +12,11 @@ public:
 
 public:
 	void Construct(const FArguments& InArgs, struct FBindState *_Owner);
-	FString GetStateText() const;
+	FText GetStateText() const;
 	void OnTypeChanged(TSharedPtr<ELiveEditControllerType::Type> InType, ESelectInfo::Type);
 	TSharedRef<SWidget> GenerateControlTypeItem(TSharedPtr<ELiveEditControllerType::Type> InType);
-	FString GetControlTypeText(ELiveEditControllerType::Type Type) const;
-	FString GetSelectedControlTypeText() const;
+	FText GetControlTypeText(ELiveEditControllerType::Type Type) const;
+	FText GetSelectedControlTypeText() const;
 
 private:
 	struct FBindState *Owner;
@@ -39,16 +41,16 @@ struct FBindState : FLiveEditorWizardBase::FState
 	{
 	}
 	
-	virtual FString GetStateTitle() const
+	virtual FText GetStateTitle() const
 	{
-		return FString( TEXT("MIDI Binding") );
+		return LOCTEXT("MIDIBinding", "MIDI Binding");
 	}
-	virtual FString GetStateText() const
+	virtual FText GetStateText() const
 	{
-		return FString::Printf( TEXT("Use the MIDI Control you would like bound to:\n\n%s\n\nChannel: %d\nControlID: %d"),
-			*Description,
-			BindingData.Channel,
-			BindingData.ControlID );
+		return FText::Format( LOCTEXT("MIDIBindingDescFmt", "Use the MIDI Control you would like bound to:\n\n{0}\n\nChannel: {1}\nControlID: {2}"),
+			FText::FromString(Description),
+			FText::AsNumber(BindingData.Channel),
+			FText::AsNumber(BindingData.ControlID) );
 	}
 
 	virtual void Init()
@@ -234,7 +236,7 @@ void SLiveEditorBlueprintBindingWindow::Construct(const FArguments& InArgs, FBin
 	];
 }
 
-FString SLiveEditorBlueprintBindingWindow::GetStateText() const
+FText SLiveEditorBlueprintBindingWindow::GetStateText() const
 {
 	return Owner->GetStateText();
 }
@@ -250,26 +252,26 @@ TSharedRef<SWidget> SLiveEditorBlueprintBindingWindow::GenerateControlTypeItem(T
 		.Text( GetControlTypeText(*InType) );
 }
 
-FString SLiveEditorBlueprintBindingWindow::GetSelectedControlTypeText() const
+FText SLiveEditorBlueprintBindingWindow::GetSelectedControlTypeText() const
 {
 	return GetControlTypeText(Owner->GetBindingData().ControlType);
 }
 
-FString SLiveEditorBlueprintBindingWindow::GetControlTypeText(ELiveEditControllerType::Type Type) const
+FText SLiveEditorBlueprintBindingWindow::GetControlTypeText(ELiveEditControllerType::Type Type) const
 {
 	switch( Type )
 	{
 	case ELiveEditControllerType::NotSet:
-		return FString(TEXT("Not Set"));
+		return LOCTEXT("NotSet", "Not Set");
 	case ELiveEditControllerType::NoteOnOff:
-		return FString(TEXT("Note On/Off"));
+		return LOCTEXT("NoteOnOrOff", "Note On/Off");
 	case ELiveEditControllerType::ControlChangeContinuous:
-		return FString(TEXT("Continuous Knob"));
+		return LOCTEXT("ControlChangeContinuous", "Continuous Knob");
 	case ELiveEditControllerType::ControlChangeFixed:
-		return FString(TEXT("Fixed Knob/Slider"));
+		return LOCTEXT("ControlChangeFixed", "Fixed Knob/Slider");
 	}
 
-	return FString();
+	return FText::GetEmpty();
 }
 
 
@@ -346,3 +348,5 @@ bool FLiveEditorBlueprintBindingWizard::GenerateStates()
 	SetFinalState(FoundCount);
 	return (FoundCount > 0);
 }
+
+#undef LOCTEXT_NAMESPACE
diff --git a/Engine/Plugins/Experimental/LiveEditor/Source/LiveEditor/Private/LiveEditorConfigWindow.cpp b/Engine/Plugins/Experimental/LiveEditor/Source/LiveEditor/Private/LiveEditorConfigWindow.cpp
--- a/Engine/Plugins/Experimental/LiveEditor/Source/LiveEditor/Private/LiveEditorConfigWindow.cpp
+++ b/Engine/Plugins/Experimental/LiveEditor/Source/LiveEditor/Private/LiveEditorConfigWindow.cpp
@@ -10,6 +10,8 @@
 
 DEFINE_LOG_CATEGORY_STATIC(LiveEditorConfigWindow, Log, All);
 
+#define LOCTEXT_NAMESPACE "LiveEditorConfig"
+
 class SLiveEditorBlueprint : public SCompoundWidget
 {
 public:
@@ -46,7 +48,7 @@ public:
 			[
 				SNew( SButton )
 				.Text( this, &SLiveEditorBlueprint::GetBindButtonText )
-				.ToolTipText( FString(TEXT("Bind your MIDI hardware to this Blueprint")) )
+				.ToolTipText( LOCTEXT("BindHardwareTooltip", "Bind your MIDI hardware to this Blueprint") )
 				.HAlign( HAlign_Center )
 				.VAlign( VAlign_Center )
 				.OnClicked( this, &SLiveEditorBlueprint::BindBlueprint )
@@ -56,8 +58,8 @@ public:
 			.Padding(2.0f)
 			[
 				SNew( SButton )
-				.Text( FString(TEXT("Activate")) )
-				.ToolTipText( FString(TEXT("Make this Blueprint active in the LiveEditor context")) )
+				.Text( LOCTEXT("Activate", "Activate") )
+				.ToolTipText( LOCTEXT("ActivateTooltip", "Make this Blueprint active in the LiveEditor context") )
 				.HAlign( HAlign_Center )
 				.VAlign( VAlign_Center )
 				.Visibility( this, &SLiveEditorBlueprint::CanActivate )
@@ -68,8 +70,8 @@ public:
 			.Padding(2.0f)
 			[
 				SNew( SButton )
-				.Text( FString(TEXT("DeActivate")) )
-				.ToolTipText( FString(TEXT("Stop this Blueprint running in the LiveEditor context")) )
+				.Text( LOCTEXT("Deactivate", "Deactivate") )
+				.ToolTipText( LOCTEXT("DeactivateTooltip", "Stop this Blueprint running in the LiveEditor context") )
 				.HAlign( HAlign_Center )
 				.VAlign( VAlign_Center )
 				.Visibility( this, &SLiveEditorBlueprint::CanDeActivate )
@@ -80,8 +82,8 @@ public:
 			.Padding(2.0f)
 			[
 				SNew( SButton )
-				.Text( FString(TEXT("Remove")) )
-				.ToolTipText( FString(TEXT("Remove this Blueprint from the LiveEditor context")) )
+				.Text( LOCTEXT("Remove", "Remove") )
+				.ToolTipText( LOCTEXT("RemoveTooltip", "Remove this Blueprint from the LiveEditor context") )
 				.HAlign( HAlign_Center )
 				.VAlign( VAlign_Center )
 				.OnClicked( this, &SLiveEditorBlueprint::RemoveBlueprint )
@@ -92,17 +94,18 @@ public:
 			.Padding(2.0f)
 			[
 				SAssignNew( NameView, STextBlock )
-				.Text( *Name.Get() )
+				.Text( FText::FromString(*Name) )
 				.ColorAndOpacity( GetFontColor() )
 			]
 		];
 	}
 
-	FString GetBindButtonText() const
+	FText GetBindButtonText() const
 	{
 		check( Blueprint != NULL );
-		return (FLiveEditorManager::Get().BlueprintHasBinding(*Blueprint))? FString(TEXT("Re-Bind")) : FString(TEXT("Bind"));
+		return (FLiveEditorManager::Get().BlueprintHasBinding(*Blueprint))? LOCTEXT("Rebind", "Re-Bind") : LOCTEXT("Bind", "Bind");
 	}
+
 	FReply BindBlueprint()
 	{
 		check( Blueprint != NULL );
@@ -235,7 +238,7 @@ public:
 				[
 					SNew( SButton )
 					.Text( this, &SLiveEditorDeviceData::GetConfigureButtonTitle )
-					.ToolTipText( FString(TEXT("Configure this MIDI hardware for use in the LiveEditor")) )
+					.ToolTipText( LOCTEXT("ConfigureHardwareTooltip", "Configure this MIDI hardware for use in the LiveEditor") )
 					.HAlign( HAlign_Center )
 					.VAlign( VAlign_Center )
 					.OnClicked( this, &SLiveEditorDeviceData::ConfigureDevice )
@@ -311,25 +314,29 @@ public:
 		return ( DeviceData.ConfigState == FLiveEditorDeviceData::CONFIGURED )? EVisibility::Visible : EVisibility::Hidden;
 	}
 
-	FString GetConfigureButtonTitle() const
+	FText GetConfigureButtonTitle() const
 	{
-		return ( DeviceData.ConfigState == FLiveEditorDeviceData::CONFIGURED )? FString(TEXT("Re-Configure")) : FString(TEXT("Configure"));
+		return ( DeviceData.ConfigState == FLiveEditorDeviceData::CONFIGURED )? LOCTEXT("Reconfigure", "Re-Configure") : LOCTEXT("Configure", "Configure");
 	}
-	FString GetTextButtonDown() const
+
+	FText GetTextButtonDown() const
 	{
-		return FString::Printf( TEXT("ButtonDown Signal: %d"), DeviceData.ButtonSignalDown );
+		return FText::Format( LOCTEXT("ButtonDownSignalFmt", "ButtonDown Signal: {0}"), FText::AsNumber(DeviceData.ButtonSignalDown) );
 	}
-	FString GetTextButtonUp() const
+
+	FText GetTextButtonUp() const
 	{
-		return FString::Printf( TEXT("ButtonUp Signal: %d"), DeviceData.ButtonSignalUp );
+		return FText::Format( LOCTEXT("ButtonUpSignalFmt", "ButtonUp Signal: {0}"), FText::AsNumber(DeviceData.ButtonSignalUp) );
 	}
-	FString GetTextCCIncrement() const
+
+	FText GetTextCCIncrement() const
 	{
-		return FString::Printf( TEXT("ContinuousIncrement Signal: %d"), DeviceData.ContinuousIncrement );
+		return FText::Format( LOCTEXT("ContinuousIncrementSignalFmt", "ContinuousIncrement Signal: {0}"), FText::AsNumber(DeviceData.ContinuousIncrement) );
 	}
-	FString GetTextCCDecrement() const
+
+	FText GetTextCCDecrement() const
 	{
-		return FString::Printf( TEXT("ContinuousDecrement Signal: %d"), DeviceData.ContinuousDecrement );
+		return FText::Format( LOCTEXT("ContinuousDecrementSignalFmt", "ContinuousDecrement Signal: {0}"), FText::AsNumber(DeviceData.ContinuousDecrement) );
 	}
 
 private:
@@ -418,7 +425,7 @@ void SLiveEditorConfigWindow::Construct(const FArguments& InArgs)
 				.HeaderContent()
 				[
 					SNew( STextBlock )
-					.Text( FString(TEXT("Blueprints")) )
+					.Text( LOCTEXT("Blueprints", "Blueprints") )
 					.Font( FEditorStyle::GetFontStyle("ContentBrowser.SourceTitleFont") )
 				]
 				.BodyContent()
@@ -432,8 +439,8 @@ void SLiveEditorConfigWindow::Construct(const FArguments& InArgs)
 						.Padding(2.0f)
 						[
 							SNew( SButton )
-							.Text( FString(TEXT("Select Blueprint")) )
-							.ToolTipText( FString(TEXT("Use Selected Asset from Content Browser")) )
+							.Text( LOCTEXT("SelectBlueprint", "Select Blueprint") )
+							.ToolTipText( LOCTEXT("SelectBlueprintTooltip", "Use Selected Asset from Content Browser") )
 							.OnClicked( this, &SLiveEditorConfigWindow::AddBlueprintToList )
 						]
 						+SVerticalBox::Slot()
@@ -458,7 +465,7 @@ void SLiveEditorConfigWindow::Construct(const FArguments& InArgs)
 				.HeaderContent()
 				[
 					SNew( STextBlock )
-					.Text( FString(TEXT("Remote Connection Manager")) )
+					.Text( LOCTEXT("RemoteConnectionManager", "Remote Connection Manager") )
 					.Font( FEditorStyle::GetFontStyle("ContentBrowser.SourceTitleFont") )
 				]
 				.BodyContent()
@@ -481,7 +488,7 @@ void SLiveEditorConfigWindow::Construct(const FArguments& InArgs)
 				.HeaderContent()
 				[
 					SNew( STextBlock )
-					.Text( FString(TEXT("Devices")) )
+					.Text( LOCTEXT("Devices", "Devices") )
 					.Font( FEditorStyle::GetFontStyle("ContentBrowser.SourceTitleFont") )
 				]
 				.BodyContent()
@@ -495,8 +502,8 @@ void SLiveEditorConfigWindow::Construct(const FArguments& InArgs)
 						.Padding(2.0f)
 						[
 							SNew( SButton )
-							.Text( FString(TEXT("Refresh Devices")) )
-							.ToolTipText( FString(TEXT("Refresh MIDI Device Connections")) )
+							.Text( LOCTEXT("RefreshDevices", "Refresh Devices") )
+							.ToolTipText( LOCTEXT("RefreshDevicesTooltip", "Refresh MIDI Device Connections") )
 							.OnClicked( this, &SLiveEditorConfigWindow::RefreshDevices )
 						]
 						+SVerticalBox::Slot()
@@ -564,3 +571,5 @@ TSharedRef<ITableRow> SLiveEditorConfigWindow::OnGenerateWidgetForDeviceList(TSh
 			SNew(SLiveEditorDeviceData, Item, this)
 		];
 }
+
+#undef LOCTEXT_NAMESPACE
diff --git a/Engine/Plugins/Experimental/LiveEditor/Source/LiveEditor/Private/LiveEditorConnectionWindow.cpp b/Engine/Plugins/Experimental/LiveEditor/Source/LiveEditor/Private/LiveEditorConnectionWindow.cpp
--- a/Engine/Plugins/Experimental/LiveEditor/Source/LiveEditor/Private/LiveEditorConnectionWindow.cpp
+++ b/Engine/Plugins/Experimental/LiveEditor/Source/LiveEditor/Private/LiveEditorConnectionWindow.cpp
@@ -1,6 +1,8 @@
 #include "LiveEditorPrivatePCH.h"
 #include "LiveEditorConnectionWindow.h"
 
+#define LOCTEXT_NAMESPACE "LiveEditorConnection"
+
 class SConnectionDetailWindow : public SCompoundWidget
 {
 public:
@@ -22,8 +24,8 @@ public:
 			.Padding(2.0f)
 			[
 				SNew( SButton )
-				.Text( FString(TEXT("Disconnect")) )
-				.ToolTipText( FString(TEXT("Disconnect from this IP")) )
+				.Text( LOCTEXT("Disconnect", "Disconnect") )
+				.ToolTipText( LOCTEXT("DisconnectTooltip", "Disconnect from this IP") )
 				.OnClicked( this, &SConnectionDetailWindow::OnDisconnect )
 			]
 			+SHorizontalBox::Slot()
@@ -69,8 +71,8 @@ void SLiveEditorConnectionWindow::Construct(const FArguments& InArgs)
 			.Padding(2.0f)
 			[
 				SNew( SButton )
-				.Text( FString(TEXT("Connect")) )
-				.ToolTipText( FString(TEXT("Connect to this IP for Remote LiveEditing")) )
+				.Text( LOCTEXT("Connect", "Connect") )
+				.ToolTipText( LOCTEXT("ConnectTooltip", "Connect to this IP for Remote LiveEditing") )
 				.OnClicked( this, &SLiveEditorConnectionWindow::OnConnect )
 			]
 			+SHorizontalBox::Slot()
@@ -83,7 +85,7 @@ void SLiveEditorConnectionWindow::Construct(const FArguments& InArgs)
 					SNew(SEditableTextBox)
 					.SelectAllTextWhenFocused(true)
 					.Text(this, &SLiveEditorConnectionWindow::GetNewHostText)
-					.ToolTipText( FString(TEXT("The host IP you would like to connect to for LiveEditing. Usage xxx.xxx.xxx.xxx (no port)")) )
+					.ToolTipText( LOCTEXT("HostnameTooltip", "The host IP you would like to connect to for LiveEditing. Usage xxx.xxx.xxx.xxx (no port)") )
 					.OnTextCommitted(this, &SLiveEditorConnectionWindow::OnNewHostTextCommited)
 					.OnTextChanged(this, &SLiveEditorConnectionWindow::OnNewHostTextCommited, ETextCommit::Default)
 				]
@@ -137,7 +139,7 @@ void SLiveEditorConnectionWindow::DisconnectFrom( const FString &IPAddress )
 
 FText SLiveEditorConnectionWindow::GetNewHostText() const
 {
-	return ( NewConnectionString != FString(TEXT("")) )? FText::FromString(NewConnectionString) : FText::FromString(FString(TEXT("Enter IP")));
+	return ( !NewConnectionString.IsEmpty() )? FText::FromString(NewConnectionString) : LOCTEXT("EnterIP", "Enter IP");
 }
 
 void SLiveEditorConnectionWindow::OnNewHostTextCommited(const FText& InText, ETextCommit::Type InCommitType)
@@ -152,3 +154,5 @@ void SLiveEditorConnectionWindow::RefreshConnections()
 		ConnectionListView->RequestListRefresh();
 	}
 }
+
+#undef LOCTEXT_NAMESPACE
diff --git a/Engine/Plugins/Experimental/LiveEditor/Source/LiveEditor/Private/LiveEditorDeviceSetupWindow.cpp b/Engine/Plugins/Experimental/LiveEditor/Source/LiveEditor/Private/LiveEditorDeviceSetupWindow.cpp
--- a/Engine/Plugins/Experimental/LiveEditor/Source/LiveEditor/Private/LiveEditorDeviceSetupWindow.cpp
+++ b/Engine/Plugins/Experimental/LiveEditor/Source/LiveEditor/Private/LiveEditorDeviceSetupWindow.cpp
@@ -7,6 +7,8 @@
 
 DEFINE_LOG_CATEGORY_STATIC(LiveEditorDeviceSetupWindow, Log, All);
 
+#define LOCTEXT_NAMESPACE "LiveEditorWizard"
+
 void SLiveEditorWizardWindow::Construct(const FArguments& InArgs)
 {
 	ChildSlot
@@ -30,7 +32,7 @@ void SLiveEditorWizardWindow::Construct(const FArguments& InArgs)
 			.Padding( 2.0f )
 			[
 				SNew( STextBlock )
-				.Text( FString(TEXT("No Content")) )
+				.Text( LOCTEXT("NoContent", "No Content") )
 				.ColorAndOpacity( FLinearColor::White )
 			]
 		]
@@ -46,7 +48,7 @@ void SLiveEditorWizardWindow::Construct(const FArguments& InArgs)
 			.HAlign(HAlign_Right)
 			[
 				SNew( SButton )
-				.Text( FString(TEXT("Cancel")) )
+				.Text( LOCTEXT("Cancel", "Cancel") )
 				.Visibility( this, &SLiveEditorWizardWindow::ShowCancelButton )
 				.OnClicked( this, &SLiveEditorWizardWindow::OnCancelClicked )
 			]
@@ -67,7 +69,7 @@ void SLiveEditorWizardWindow::Activated()
 	GenerateStateContent();
 }
 
-FString SLiveEditorWizardWindow::GetTitle() const
+FText SLiveEditorWizardWindow::GetTitle() const
 {
 	FLiveEditorWizardBase *Wizard = FLiveEditorManager::Get().GetActiveWizard();
 	check( Wizard != NULL );
@@ -85,12 +87,12 @@ void SLiveEditorWizardWindow::GenerateStateContent()
 	{
 		DynamicContentPane.Get()->SetContent(
 			SNew( STextBlock )
-			.Text( FString(TEXT("No Content")) )
+			.Text( LOCTEXT("NoContent", "No Content") )
 			.ColorAndOpacity( FLinearColor::White ) );
 	}
 }
 
-FString SLiveEditorWizardWindow::GetButtonText() const
+FText SLiveEditorWizardWindow::GetButtonText() const
 {
 	FLiveEditorWizardBase *Wizard = FLiveEditorManager::Get().GetActiveWizard();
 	check( Wizard != NULL );
@@ -129,3 +131,5 @@ FReply SLiveEditorWizardWindow::OnButtonClicked()
 	GenerateStateContent();
 	return FReply::Handled();
 }
+
+#undef LOCTEXT_NAMESPACE
diff --git a/Engine/Plugins/Experimental/LiveEditor/Source/LiveEditor/Private/LiveEditorDeviceSetupWindow.h b/Engine/Plugins/Experimental/LiveEditor/Source/LiveEditor/Private/LiveEditorDeviceSetupWindow.h
--- a/Engine/Plugins/Experimental/LiveEditor/Source/LiveEditor/Private/LiveEditorDeviceSetupWindow.h
+++ b/Engine/Plugins/Experimental/LiveEditor/Source/LiveEditor/Private/LiveEditorDeviceSetupWindow.h
@@ -18,8 +18,8 @@ public:
 private:
 	void GenerateStateContent();
 
-	FString GetTitle() const;
-	FString GetButtonText() const;
+	FText GetTitle() const;
+	FText GetButtonText() const;
 
 	EVisibility ShowCancelButton() const;
 	bool ButtonEnabled() const;
diff --git a/Engine/Plugins/Experimental/LiveEditor/Source/LiveEditor/Private/LiveEditorDeviceSetupWizard.cpp b/Engine/Plugins/Experimental/LiveEditor/Source/LiveEditor/Private/LiveEditorDeviceSetupWizard.cpp
--- a/Engine/Plugins/Experimental/LiveEditor/Source/LiveEditor/Private/LiveEditorDeviceSetupWizard.cpp
+++ b/Engine/Plugins/Experimental/LiveEditor/Source/LiveEditor/Private/LiveEditorDeviceSetupWizard.cpp
@@ -1,6 +1,8 @@
 #include "LiveEditorPrivatePCH.h"
 #include "LiveEditorDeviceSetupWizard.h"
 
+#define LOCTEXT_NAMESPACE "LiveEditorWizard"
+
 namespace nLiveEditorDeviceSetupWizard
 {
 	enum States
@@ -49,13 +51,13 @@ struct FConfigurationState : public FLiveEditorDeviceSetupWizard::FState
 	{
 	}
 
-	virtual FString GetStateTitle() const
+	virtual FText GetStateTitle() const
 	{
-		return FString( TEXT("Device Questionnaire") );
+		return LOCTEXT("DeviceQuestionnaire", "Device Questionnaire");
 	}
-	virtual FString GetStateText() const
+	virtual FText GetStateText() const
 	{
-		return FString( TEXT("Select all properties that apply to your device") );
+		return LOCTEXT("DeviceQuestionnaireDesc", "Select all properties that apply to your device");
 	}
 	virtual void Init()
 	{
@@ -121,7 +123,7 @@ void SDeviceQuestionnaireSubWindow::Construct(const FArguments& InArgs, FConfigu
 			.Content()
 			[
 				SNew(STextBlock)
-				.Text( FString(TEXT("Buttons")) )
+				.Text( LOCTEXT("Buttons", "Buttons") )
 			]
 		]
 		+SVerticalBox::Slot()
@@ -134,7 +136,7 @@ void SDeviceQuestionnaireSubWindow::Construct(const FArguments& InArgs, FConfigu
 			.Content()
 			[
 				SNew(STextBlock)
-				.Text( FString(TEXT("Endless Encoders")) )
+				.Text( LOCTEXT("EndlessEncoders", "Endless Encoders") )
 			]
 		]
 	];
@@ -168,13 +170,13 @@ struct FButtonState : public FLiveEditorDeviceSetupWizard::FState
 		HandleType(HT_DOWN)
 	{
 	}
-	virtual FString GetStateTitle() const
+	virtual FText GetStateTitle() const
 	{
-		return FString( TEXT("Button Configuration") );
+		return LOCTEXT("ButtonConfiguration", "Button Configuration");
 	}
-	virtual FString GetStateText() const
+	virtual FText GetStateText() const
 	{
-		return FString( TEXT("Push a button") );
+		return LOCTEXT("ButtonConfigurationDesc", "Push a button");
 	}
 	virtual void Init()
 	{
@@ -277,13 +279,13 @@ private:
 struct FContinuousKnobRightState : public FContinuousKnobBaseState
 {
 	FContinuousKnobRightState( int32 _NextState ) : FContinuousKnobBaseState( 15, _NextState ) {}
-	virtual FString GetStateTitle() const
+	virtual FText GetStateTitle() const
 	{
-		return FString( TEXT("Continuous Knob Configuration") );
+		return LOCTEXT("ContinuousKnobConfiguration", "Continuous Knob Configuration");
 	}
-	virtual FString GetStateText() const
+	virtual FText GetStateText() const
 	{
-		return FString( TEXT("Twist a knob to the RIGHT") );
+		return LOCTEXT("ContinuousKnobConfigurationRightDesc", "Twist a knob to the RIGHT");
 	}
 
 protected:
@@ -296,13 +298,13 @@ protected:
 struct FContinuousKnobLeftState : public FContinuousKnobBaseState
 {
 	FContinuousKnobLeftState( int32 _NextState ) : FContinuousKnobBaseState( 15, _NextState ) {}
-	virtual FString GetStateTitle() const
+	virtual FText GetStateTitle() const
 	{
-		return FString( TEXT("Continuous Knob Configuration") );
+		return LOCTEXT("ContinuousKnobConfiguration", "Continuous Knob Configuration");
 	}
-	virtual FString GetStateText() const
+	virtual FText GetStateText() const
 	{
-		return FString( TEXT("Twist a knob to the LEFT") );
+		return LOCTEXT("ContinuousKnobConfigurationLeftDesc", "Twist a knob to the LEFT");
 	}
 
 protected:
@@ -328,12 +330,12 @@ FLiveEditorDeviceSetupWizard::~FLiveEditorDeviceSetupWizard()
 {
 }
 
-FString FLiveEditorDeviceSetupWizard::GetAdvanceButtonText() const
+FText FLiveEditorDeviceSetupWizard::GetAdvanceButtonText() const
 {
 	if ( GetCurState() == nLiveEditorDeviceSetupWizard::S_CONFIGURATION )
 	{
 		check( ConfigState != NULL );
-		return ( ConfigState->bHasButtons || ConfigState->bHasEndlessEncoders )? FString(TEXT("Next")) : FString(TEXT("Finish"));
+		return ( ConfigState->bHasButtons || ConfigState->bHasEndlessEncoders )? LOCTEXT("Next", "Next") : LOCTEXT("Finish", "Finish");
 	}
 	else
 	{
@@ -374,3 +376,4 @@ void FLiveEditorDeviceSetupWizard::OnWizardFinished( struct FLiveEditorDeviceDat
 	FLiveEditorManager::SaveDeviceData( Data );
 }
 
+#undef LOCTEXT_NAMESPACE
diff --git a/Engine/Plugins/Experimental/LiveEditor/Source/LiveEditor/Private/LiveEditorDeviceSetupWizard.h b/Engine/Plugins/Experimental/LiveEditor/Source/LiveEditor/Private/LiveEditorDeviceSetupWizard.h
--- a/Engine/Plugins/Experimental/LiveEditor/Source/LiveEditor/Private/LiveEditorDeviceSetupWizard.h
+++ b/Engine/Plugins/Experimental/LiveEditor/Source/LiveEditor/Private/LiveEditorDeviceSetupWizard.h
@@ -14,7 +14,7 @@ public:
 	void Start( PmDeviceID _DeviceID, struct FLiveEditorDeviceData &Data );
 	void Configure( bool bHasButtons, bool bHasEndlessEncoders );
 
-	virtual FString GetAdvanceButtonText() const;
+	virtual FText GetAdvanceButtonText() const override;
 
 protected:
 	virtual void OnWizardFinished( struct FLiveEditorDeviceData &Data );
diff --git a/Engine/Plugins/Experimental/LiveEditor/Source/LiveEditor/Private/LiveEditorWizardBase.cpp b/Engine/Plugins/Experimental/LiveEditor/Source/LiveEditor/Private/LiveEditorWizardBase.cpp
--- a/Engine/Plugins/Experimental/LiveEditor/Source/LiveEditor/Private/LiveEditorWizardBase.cpp
+++ b/Engine/Plugins/Experimental/LiveEditor/Source/LiveEditor/Private/LiveEditorWizardBase.cpp
@@ -1,6 +1,8 @@
 #include "LiveEditorPrivatePCH.h"
 #include "LiveEditorWizardBase.h"
 
+#define LOCTEXT_NAMESPACE "LiveEditorWizard"
+
 /**
  *
  * FLiveEditorDeviceSetupWizard::FState
@@ -40,7 +42,7 @@ FLiveEditorWizardBase::~FLiveEditorWizardBase()
 	ClearStates();
 }
 
-FString FLiveEditorWizardBase::GetActiveStateTitle() const
+FText FLiveEditorWizardBase::GetActiveStateTitle() const
 {
 	if ( States.Contains(CurState) )
 	{
@@ -48,7 +50,7 @@ FString FLiveEditorWizardBase::GetActiveStateTitle() const
 	}
 	else
 	{
-		return FString( TEXT("Complete") );
+		return LOCTEXT("Complete", "Complete");
 	}
 }
 
@@ -174,7 +176,9 @@ void FLiveEditorWizardBase::ClearStates()
 	States.Empty();
 }
 
-FString FLiveEditorWizardBase::GetAdvanceButtonText() const
+FText FLiveEditorWizardBase::GetAdvanceButtonText() const
 {
-	return (IsOnLastStep())? FString( TEXT("Finish") ) : FString( TEXT("Next") );
+	return (IsOnLastStep())? LOCTEXT("Finish", "Finish") : LOCTEXT("Next", "Next");
 }
+
+#undef LOCTEXT_NAMESPACE
diff --git a/Engine/Plugins/Experimental/LiveEditor/Source/LiveEditor/Private/LiveEditorWizardBase.h b/Engine/Plugins/Experimental/LiveEditor/Source/LiveEditor/Private/LiveEditorWizardBase.h
--- a/Engine/Plugins/Experimental/LiveEditor/Source/LiveEditor/Private/LiveEditorWizardBase.h
+++ b/Engine/Plugins/Experimental/LiveEditor/Source/LiveEditor/Private/LiveEditorWizardBase.h
@@ -11,7 +11,7 @@ public:
 	FLiveEditorWizardBase( int32 _FinalState );
 	virtual ~FLiveEditorWizardBase();
 
-	FString GetActiveStateTitle() const;
+	FText GetActiveStateTitle() const;
 
 	// _DeviceID: device to which this Wizard is restricted from recieving data. If == pmNoDevice, all device input is permitted
 	void Start( int32 FirstState, PmDeviceID _RestrictedToDeviceID = pmNoDevice );
@@ -39,7 +39,7 @@ public:
 		bReadyToAdvance = true;
 	}
 
-	virtual FString GetAdvanceButtonText() const;
+	virtual FText GetAdvanceButtonText() const;
 
 	struct FState
 	{
@@ -47,8 +47,8 @@ public:
 		virtual ~FState() {}
 		int32 GetNextState() const;
 
-		virtual FString GetStateTitle() const = 0;
-		virtual FString GetStateText() const = 0;
+		virtual FText GetStateTitle() const = 0;
+		virtual FText GetStateText() const = 0;
 		
 		virtual void Init() = 0;
 		virtual void OnExit() {}
