Fixing build issue with LiveEditor related to changes made in CL# 2580309

diff --git a/Engine/Plugins/Experimental/LiveEditor/Source/LiveEditor/Private/SGraphPinLiveEditVar.cpp b/Engine/Plugins/Experimental/LiveEditor/Source/LiveEditor/Private/SGraphPinLiveEditVar.cpp
--- a/Engine/Plugins/Experimental/LiveEditor/Source/LiveEditor/Private/SGraphPinLiveEditVar.cpp
+++ b/Engine/Plugins/Experimental/LiveEditor/Source/LiveEditor/Private/SGraphPinLiveEditVar.cpp
@@ -25,22 +25,22 @@ TSharedRef<SWidget>	SGraphPinLiveEditVar::GetDefaultValueWidget()
 		.OnGetDisplayName(this, &SGraphPinLiveEditVar::OnGetFriendlyName);
 }
 
-FString SGraphPinLiveEditVar::OnGetFriendlyName(int32 EnumIndex) const
+FText SGraphPinLiveEditVar::OnGetFriendlyName(int32 EnumIndex) const
 {
 	if ( EnumIndex >= 0 && EnumIndex < VariableNameList.Num() )
 	{
-		return VariableNameList[EnumIndex];
+		return FText::FromString(VariableNameList[EnumIndex]);
 	}
 	else
 	{
-		return FString::Printf( TEXT("BOOM - should never happen"), EnumIndex );
+		return FText::FromString(FString::Printf( TEXT("BOOM - should never happen"), EnumIndex ));
 	}
 }
 
 void SGraphPinLiveEditVar::ComboBoxSelectionChanged( TSharedPtr<int32> NewSelection, ESelectInfo::Type /*SelectInfo*/ )
 {
 	//Set new selection
-	CurrentValue = OnGetFriendlyName( *NewSelection.Get() );
+	CurrentValue = OnGetFriendlyName( *NewSelection.Get() ).ToString();
 	GraphPinObj->GetSchema()->TrySetDefaultValue(*GraphPinObj, CurrentValue);
 }
 
diff --git a/Engine/Plugins/Experimental/LiveEditor/Source/LiveEditor/Private/SGraphPinLiveEditVar.h b/Engine/Plugins/Experimental/LiveEditor/Source/LiveEditor/Private/SGraphPinLiveEditVar.h
--- a/Engine/Plugins/Experimental/LiveEditor/Source/LiveEditor/Private/SGraphPinLiveEditVar.h
+++ b/Engine/Plugins/Experimental/LiveEditor/Source/LiveEditor/Private/SGraphPinLiveEditVar.h
@@ -52,7 +52,7 @@ private:
 	 *
 	 * @param EnumIndex	- The index of the enum to return the friendly name for
 	 */
-	FString OnGetFriendlyName(int32 EnumIndex) const;
+	FText OnGetFriendlyName(int32 EnumIndex) const;
 
 	/**
 	 * Executes a series of checks to determine if this is a variable that the LiveEditor should really have access to
